HoldoutComposite: Moved plugin to Experimental/Compositing folder.

#rnx

diff --git a/Engine/Plugins/Experimental/Compositing/HoldoutComposite/HoldoutComposite.uplugin b/Engine/Plugins/Experimental/Compositing/HoldoutComposite/HoldoutComposite.uplugin
new file mode 100644
--- /dev/null
+++ b/Engine/Plugins/Experimental/Compositing/HoldoutComposite/HoldoutComposite.uplugin
@@ -0,0 +1,25 @@
+{
+	"FileVersion": 1,
+	"Version": 1,
+	"VersionName": "1.0",
+	"FriendlyName": "HoldoutComposite",
+	"Description": "Plugin designed to facilitate alpha holdout compositing primitives after post-processing.",
+	"Category": "Compositing",
+	"CreatedBy": "Epic Games, Inc.",
+	"CreatedByURL": "https://epicgames.com",
+	"DocsURL": "",
+	"MarketplaceURL": "",
+	"SupportURL": "",
+	"EnabledByDefault": false,
+	"CanContainContent": false,
+	"IsBetaVersion": false,
+	"IsExperimentalVersion": true,
+	"Installed": false,
+	"Modules": [
+		{
+			"Name": "HoldoutComposite",
+			"Type": "Runtime",
+			"LoadingPhase": "PostConfigInit"
+		}
+	]
+}
\ No newline at end of file
diff --git a/Engine/Plugins/Experimental/Compositing/HoldoutComposite/Shaders/Private/HoldoutCompositeDilate.usf b/Engine/Plugins/Experimental/Compositing/HoldoutComposite/Shaders/Private/HoldoutCompositeDilate.usf
new file mode 100644
--- /dev/null
+++ b/Engine/Plugins/Experimental/Compositing/HoldoutComposite/Shaders/Private/HoldoutCompositeDilate.usf
@@ -0,0 +1,114 @@
+// Copyright Epic Games, Inc. All Rights Reserved.
+
+#include "/Engine/Private/Common.ush"
+
+#ifndef TILE_SIZE
+    #error TILE_SIZE must be defined
+#endif
+
+#define KERNEL_COUNT 8
+#define ALPHA_THRESHOLD 0.5f // Custom render passes are opaque-only currently
+
+// Number of texels arround the group tile for scene color.
+#define LDS_TILE_BORDER_SIZE (1)
+
+// Width in texels of the depth tile cached into LDS.
+#define LDS_TILE_WIDTH (TILE_SIZE + 2 * LDS_TILE_BORDER_SIZE)
+
+// Total number of texels cached in the scene color tile.
+#define LDS_ARRAY_SIZE (LDS_TILE_WIDTH * LDS_TILE_WIDTH)
+
+#define THREADGROUP_TOTAL (TILE_SIZE * TILE_SIZE)
+
+Texture2D InputTexture;
+RWTexture2D<float4> RWOutputTexture;
+uint2 Dimensions;
+
+static const int2 kOffsets[KERNEL_COUNT] = {
+    int2(-1,  0),
+    int2( 1,  0),
+    int2( 0,  1),
+    int2( 0, -1),
+    int2(-1,  1),
+    int2( 1,  1),
+    int2( 1, -1),
+    int2(-1, -1),
+};
+
+groupshared float4 SharedColor[LDS_ARRAY_SIZE];
+
+// Get the index within the LDS array.
+uint GetTileArrayIndexFromPixelOffset(in uint2 GroupThreadId, in int2 PixelOffset)
+{
+    uint2 TilePos = GroupThreadId + uint2(PixelOffset + LDS_TILE_BORDER_SIZE);
+    return TilePos.x + TilePos.y * (LDS_TILE_BORDER_SIZE * 2 + TILE_SIZE);
+}
+
+float4 SampleCachedColors(in uint2 GroupThreadId, int2 PixelOffset = int2(0,0))
+{
+    return SharedColor[GetTileArrayIndexFromPixelOffset(GroupThreadId, PixelOffset)];
+}
+
+[numthreads(TILE_SIZE, TILE_SIZE, 1)]
+void MainCS(
+    uint2 GroupId : SV_GroupID,
+    uint2 GroupThreadId : SV_GroupThreadID,
+    uint GroupThreadIndex : SV_GroupIndex,
+    uint2 DispatchThreadId : SV_DispatchThreadID)
+{
+    // Init LDS
+    {
+        // Global texel coordinate of the active group (corner) offset by the LDS border
+        uint2 GroupTexelOffset = GroupId * uint2(TILE_SIZE, TILE_SIZE) - LDS_TILE_BORDER_SIZE;
+
+        // To account for the dilation margin, each warp thread should load extra texels.
+        // We proceed linearly through the tile (with borders) and each warp thread loads
+        // multiple colors at THREADGROUP_TOTAL intervals with the necessary boundary conditions.
+        const uint LoadCount = (LDS_ARRAY_SIZE + THREADGROUP_TOTAL - 1) / THREADGROUP_TOTAL; // integer ceiling
+
+        uint LinearGroupThreadId = GroupThreadIndex;
+
+        UNROLL
+        for (uint i = 0; i < LoadCount; i++)
+        {
+            uint2 TexelLocation =  GroupTexelOffset + uint2(
+                LinearGroupThreadId % LDS_TILE_WIDTH,
+                LinearGroupThreadId / LDS_TILE_WIDTH);
+
+            if ((LinearGroupThreadId < LDS_ARRAY_SIZE) ||
+                (i != LoadCount - 1) ||
+                (LDS_ARRAY_SIZE % THREADGROUP_TOTAL) == 0)
+            {
+                SharedColor[LinearGroupThreadId] = InputTexture[TexelLocation];
+            }
+            LinearGroupThreadId += THREADGROUP_TOTAL;
+        }
+    }
+
+    GroupMemoryBarrierWithGroupSync();
+
+    // Discard out-of-bound texels
+    if (any(DispatchThreadId >= Dimensions))
+    {
+        return;
+    }
+
+    float4 TexelColor = SampleCachedColors(GroupThreadId);
+    
+    BRANCH
+    if(TexelColor.a > ALPHA_THRESHOLD) // Alpha is inverted opacity
+    {
+        UNROLL
+        for(int i=0; i<KERNEL_COUNT; i++)
+        {
+            float4 OffsetColor = SampleCachedColors(GroupThreadId, kOffsets[i]);
+
+            if(OffsetColor.a < TexelColor.a)
+            {
+                TexelColor = OffsetColor;
+            }
+        }
+    }
+
+    RWOutputTexture[DispatchThreadId] = TexelColor;
+}
diff --git a/Engine/Plugins/Experimental/Compositing/HoldoutComposite/Shaders/Private/HoldoutCompositeFinal.usf b/Engine/Plugins/Experimental/Compositing/HoldoutComposite/Shaders/Private/HoldoutCompositeFinal.usf
new file mode 100644
--- /dev/null
+++ b/Engine/Plugins/Experimental/Compositing/HoldoutComposite/Shaders/Private/HoldoutCompositeFinal.usf
@@ -0,0 +1,52 @@
+// Copyright Epic Games, Inc. All Rights Reserved.
+
+#include "/Engine/Private/Common.ush"
+#include "/Engine/Private/GammaCorrectionCommon.ush"
+#include "/Engine/Private/ScreenPass.ush"
+
+SCREEN_PASS_TEXTURE_VIEWPORT(Input)
+SCREEN_PASS_TEXTURE_VIEWPORT(Custom)
+SCREEN_PASS_TEXTURE_VIEWPORT(Output)
+
+Texture2D InputTexture;
+SamplerState InputSampler;
+
+Texture2D CustomTexture;
+SamplerState CustomSampler;
+
+StructuredBuffer<float4> EyeAdaptationBuffer;
+uint SourceEncoding;
+float DisplayGamma;
+
+void MainPS(
+    float4 SvPosition : SV_POSITION,
+    out float4 OutColor : SV_Target0
+    )
+{
+    float2 ViewportUV = (SvPosition.xy - Output_ViewportMin) * Output_ViewportSizeInverse;
+    
+    const float2 InputUV = (ViewportUV * Input_ViewportSize + Input_ViewportMin) * Input_ExtentInverse;
+    float4 SceneColor = Texture2DSample(InputTexture, InputSampler, InputUV);
+
+    const float2 CustomUV = (ViewportUV * Custom_ViewportSize + Custom_ViewportMin) * Custom_ExtentInverse;
+    float4 CustomColor = Texture2DSample(CustomTexture, CustomSampler, CustomUV);
+
+#if PERMUTATION_GLOBAL_EXPOSURE
+    const float GlobalExposure = EyeAdaptationBuffer[0].x;
+    CustomColor.rgb *= GlobalExposure;
+#endif
+
+    BRANCH
+    if(SourceEncoding == 1)
+    {
+        SceneColor.rgb = pow(SceneColor.rgb, DisplayGamma);
+    }
+    else if(SourceEncoding == 2)
+    {
+        SceneColor.rgb = sRGBToLinear(SceneColor.rgb);
+    }
+
+    // TODO: HDR cases will need to be handled explicitely.
+    OutColor.rgb = LinearToSrgb(CustomColor.rgb * SceneColor.a + SceneColor.rgb);
+    OutColor.a = CustomColor.a * SceneColor.a;
+}
diff --git a/Engine/Plugins/Experimental/Compositing/HoldoutComposite/Source/HoldoutComposite/HoldoutComposite.Build.cs b/Engine/Plugins/Experimental/Compositing/HoldoutComposite/Source/HoldoutComposite/HoldoutComposite.Build.cs
new file mode 100644
--- /dev/null
+++ b/Engine/Plugins/Experimental/Compositing/HoldoutComposite/Source/HoldoutComposite/HoldoutComposite.Build.cs
@@ -0,0 +1,48 @@
+// Copyright Epic Games, Inc. All Rights Reserved.
+
+using UnrealBuildTool;
+
+public class HoldoutComposite : ModuleRules
+{
+	public HoldoutComposite(ReadOnlyTargetRules Target) : base(Target)
+	{
+		PCHUsage = ModuleRules.PCHUsageMode.UseExplicitOrSharedPCHs;
+
+		PrivateIncludePaths.AddRange(
+			new string[] {
+				System.IO.Path.Combine(GetModuleDirectory("Renderer"), "Private"),
+			}
+		);
+
+		PublicDependencyModuleNames.AddRange(
+			new string[]
+			{
+				"Core",
+				"CoreUObject",
+				"Engine",
+				"DeveloperSettings",
+			}
+		);
+
+		PrivateDependencyModuleNames.AddRange(
+			new string[]
+			{
+				"Projects",
+				"RenderCore",
+				"Renderer",
+				"RHI",
+			}
+		);
+
+		if (Target.bBuildEditor)
+		{
+			PrivateDependencyModuleNames.AddRange(
+				new string[]
+				{
+					"UnrealEd",
+					"Slate",
+				}
+			);
+		}
+	}
+}
diff --git a/Engine/Plugins/Experimental/Compositing/HoldoutComposite/Source/HoldoutComposite/Private/HoldoutCompositeComponent.cpp b/Engine/Plugins/Experimental/Compositing/HoldoutComposite/Source/HoldoutComposite/Private/HoldoutCompositeComponent.cpp
new file mode 100644
--- /dev/null
+++ b/Engine/Plugins/Experimental/Compositing/HoldoutComposite/Source/HoldoutComposite/Private/HoldoutCompositeComponent.cpp
@@ -0,0 +1,93 @@
+// Copyright Epic Games, Inc. All Rights Reserved.
+
+#include "HoldoutCompositeComponent.h"
+
+#include "Components/PrimitiveComponent.h"
+#include "Engine/Engine.h"
+#include "HoldoutCompositeModule.h"
+#include "HoldoutCompositeSubsystem.h"
+
+#if WITH_EDITOR
+#include "Framework/Notifications/NotificationManager.h"
+#include "Widgets/Notifications/SNotificationList.h"
+#endif
+
+#define LOCTEXT_NAMESPACE "HoldoutComposite"
+
+UHoldoutCompositeComponent::UHoldoutCompositeComponent(const FObjectInitializer& ObjectInitializer)
+	: Super(ObjectInitializer)
+{
+}
+
+void UHoldoutCompositeComponent::OnRegister()
+{
+	Super::OnRegister();
+
+	RegisterCompositeImpl();
+}
+
+void UHoldoutCompositeComponent::OnUnregister()
+{
+	UnregisterCompositeImpl();
+
+	Super::OnUnregister();
+}
+
+void UHoldoutCompositeComponent::DetachFromComponent(const FDetachmentTransformRules& DetachmentRules)
+{
+	// Note: We also unregister here while the attached parent pointer is still valid.
+	UnregisterCompositeImpl();
+
+	Super::DetachFromComponent(DetachmentRules);
+}
+
+void UHoldoutCompositeComponent::OnAttachmentChanged()
+{
+	Super::OnAttachmentChanged();
+
+	UnregisterCompositeImpl();
+
+	const USceneComponent* SceneComponent = GetAttachParent();
+	if (IsValid(SceneComponent))
+	{
+		const UPrimitiveComponent* ParentPrimitiveComponent = Cast<UPrimitiveComponent>(SceneComponent);
+		if (IsValid(ParentPrimitiveComponent))
+		{
+			RegisterCompositeImpl();
+		}
+		else
+		{
+#if WITH_EDITOR
+			FNotificationInfo Info(LOCTEXT("CompositeParentNotification",
+				"The composite component must be parented to a primitive component."));
+			Info.ExpireDuration = 5.0f;
+
+			FSlateNotificationManager::Get().AddNotification(Info);
+#endif
+		}
+	}
+}
+
+void UHoldoutCompositeComponent::RegisterCompositeImpl()
+{
+	UHoldoutCompositeSubsystem* Subsystem = UWorld::GetSubsystem<UHoldoutCompositeSubsystem>(GetWorld());
+	UPrimitiveComponent* ParentPrimitiveComponent = Cast<UPrimitiveComponent>(GetAttachParent());
+
+	if (IsValid(Subsystem) && IsValid(ParentPrimitiveComponent))
+	{
+		Subsystem->RegisterPrimitive(ParentPrimitiveComponent);
+	}
+}
+
+void UHoldoutCompositeComponent::UnregisterCompositeImpl()
+{
+	UHoldoutCompositeSubsystem* Subsystem = UWorld::GetSubsystem<UHoldoutCompositeSubsystem>(GetWorld());
+	UPrimitiveComponent* ParentPrimitiveComponent = Cast<UPrimitiveComponent>(GetAttachParent());
+
+	if (IsValid(Subsystem) && IsValid(ParentPrimitiveComponent))
+	{
+		Subsystem->UnregisterPrimitive(ParentPrimitiveComponent);
+	}
+}
+
+#undef LOCTEXT_NAMESPACE
diff --git a/Engine/Plugins/Experimental/Compositing/HoldoutComposite/Source/HoldoutComposite/Private/HoldoutCompositeModule.cpp b/Engine/Plugins/Experimental/Compositing/HoldoutComposite/Source/HoldoutComposite/Private/HoldoutCompositeModule.cpp
new file mode 100644
--- /dev/null
+++ b/Engine/Plugins/Experimental/Compositing/HoldoutComposite/Source/HoldoutComposite/Private/HoldoutCompositeModule.cpp
@@ -0,0 +1,31 @@
+// Copyright Epic Games, Inc. All Rights Reserved.
+
+#include "HoldoutCompositeModule.h"
+
+#include "Interfaces/IPluginManager.h"
+#include "HDRHelper.h"
+#include "Misc/Paths.h"
+#include "Modules/ModuleManager.h"
+#include "ShaderCore.h"
+
+DEFINE_LOG_CATEGORY(LogHoldoutComposite);
+
+#define LOCTEXT_NAMESPACE "HoldoutComposite"
+
+void FHoldoutCompositeModule::StartupModule()
+{
+	AddShaderSourceDirectoryMapping(TEXT("/Plugin/HoldoutComposite"), FPaths::Combine(IPluginManager::Get().FindPlugin(TEXT("HoldoutComposite"))->GetBaseDir(), TEXT("Shaders")));
+
+	if (IsHDREnabled())
+	{
+		UE_LOG(LogHoldoutComposite, Warning, TEXT("Holdout composite disabled: HDR mode is not currently supported."));
+	}
+}
+
+void FHoldoutCompositeModule::ShutdownModule()
+{
+}
+
+IMPLEMENT_MODULE( FHoldoutCompositeModule, HoldoutComposite )
+
+#undef LOCTEXT_NAMESPACE
diff --git a/Engine/Plugins/Experimental/Compositing/HoldoutComposite/Source/HoldoutComposite/Private/HoldoutCompositeSceneViewExtension.cpp b/Engine/Plugins/Experimental/Compositing/HoldoutComposite/Source/HoldoutComposite/Private/HoldoutCompositeSceneViewExtension.cpp
new file mode 100644
--- /dev/null
+++ b/Engine/Plugins/Experimental/Compositing/HoldoutComposite/Source/HoldoutComposite/Private/HoldoutCompositeSceneViewExtension.cpp
@@ -0,0 +1,400 @@
+// Copyright Epic Games, Inc. All Rights Reserved.
+
+#include "HoldoutCompositeSceneViewExtension.h"
+
+#include "HoldoutCompositeSettings.h"
+#include "HoldoutCompositeModule.h"
+
+#include "Components/PrimitiveComponent.h"
+#include "Containers/Map.h"
+#include "Containers/Set.h"
+#include "EngineUtils.h"
+#include "FXRenderingUtils.h"
+#include "HDRHelper.h"
+#include "PixelShaderUtils.h"
+#include "PostProcess/PostProcessMaterialInputs.h"
+#include "Rendering/CustomRenderPass.h"
+#include "RenderGraphBuilder.h"
+#include "RenderGraphUtils.h"
+#include "SceneView.h"
+#include "ScreenPass.h"
+#include "SystemTextures.h"
+
+// Private renderer includes currently needed...
+#include "SceneRendering.h"
+#include "SceneTextureParameters.h"
+
+DECLARE_GPU_STAT_NAMED(FHoldoutCompositeDilate, TEXT("HoldoutComposite.Dilate"));
+DECLARE_GPU_STAT_NAMED(FHoldoutCompositeFinal, TEXT("HoldoutComposite.Final"));
+
+namespace HoldoutComposite
+{
+	enum class ESceneColorSourceEncoding : uint32
+	{
+		Linear = 0,
+		Gamma = 1,
+		sRGB = 2,
+	};
+}
+
+class FDilateHoldoutCompositeCS : public FGlobalShader
+{
+	DECLARE_GLOBAL_SHADER(FDilateHoldoutCompositeCS);
+	SHADER_USE_PARAMETER_STRUCT(FDilateHoldoutCompositeCS, FGlobalShader);
+
+	BEGIN_SHADER_PARAMETER_STRUCT(FParameters, )
+		SHADER_PARAMETER_RDG_TEXTURE(Texture2D, InputTexture)
+		SHADER_PARAMETER_RDG_TEXTURE_UAV(RWTexture2D<float4>, RWOutputTexture)
+		SHADER_PARAMETER(FIntPoint, Dimensions)
+	END_SHADER_PARAMETER_STRUCT()
+
+	static inline void ModifyCompilationEnvironment(const FGlobalShaderPermutationParameters& Parameters, FShaderCompilerEnvironment& OutEnvironment)
+	{
+		OutEnvironment.SetDefine(TEXT("TILE_SIZE"), ThreadGroupSize);
+
+		FGlobalShader::ModifyCompilationEnvironment(Parameters, OutEnvironment);
+	}
+
+	static const uint32 ThreadGroupSize = 16;
+};
+IMPLEMENT_GLOBAL_SHADER(FDilateHoldoutCompositeCS, "/Plugin/HoldoutComposite/Private/HoldoutCompositeDilate.usf", "MainCS", SF_Compute);
+
+class FCompositeHoldoutCompositePS : public FGlobalShader
+{
+	DECLARE_GLOBAL_SHADER(FCompositeHoldoutCompositePS);
+	SHADER_USE_PARAMETER_STRUCT(FCompositeHoldoutCompositePS, FGlobalShader);
+
+	class FApplyGlobalExposure : SHADER_PERMUTATION_BOOL("PERMUTATION_GLOBAL_EXPOSURE");
+	using FPermutationDomain = TShaderPermutationDomain<FApplyGlobalExposure>;
+
+	BEGIN_SHADER_PARAMETER_STRUCT(FParameters, )
+		SHADER_PARAMETER_STRUCT_REF(FViewUniformShaderParameters, View)
+		SHADER_PARAMETER_STRUCT(FScreenPassTextureViewportParameters, Input)
+		SHADER_PARAMETER_STRUCT(FScreenPassTextureViewportParameters, Custom)
+		SHADER_PARAMETER_STRUCT(FScreenPassTextureViewportParameters, Output)
+		SHADER_PARAMETER_RDG_TEXTURE(Texture2D, InputTexture)
+		SHADER_PARAMETER_SAMPLER(SamplerState, InputSampler)
+		SHADER_PARAMETER_RDG_TEXTURE(Texture2D, CustomTexture)
+		SHADER_PARAMETER_SAMPLER(SamplerState, CustomSampler)
+		SHADER_PARAMETER_RDG_BUFFER_SRV(StructuredBuffer<float4>, EyeAdaptationBuffer)
+		SHADER_PARAMETER(uint32, SourceEncoding)
+		SHADER_PARAMETER(float, DisplayGamma)
+
+		RENDER_TARGET_BINDING_SLOTS()
+	END_SHADER_PARAMETER_STRUCT()
+};
+IMPLEMENT_GLOBAL_SHADER(FCompositeHoldoutCompositePS, "/Plugin/HoldoutComposite/Private/HoldoutCompositeFinal.usf", "MainPS", SF_Pixel);
+
+class FHoldoutCompositeCustomRenderPass : public FCustomRenderPassBase
+{
+public:
+	IMPLEMENT_CUSTOM_RENDER_PASS(FHoldoutCompositeCustomRenderPass);
+
+	FHoldoutCompositeCustomRenderPass(const FIntPoint& InRenderTargetSize, FHoldoutCompositeSceneViewExtension* InParentExtension, FSceneView& InView)
+		: FCustomRenderPassBase(TEXT("HoldoutCompositeCustomRenderPass"), ERenderMode::DepthAndBasePass, ERenderOutput::SceneColorAndAlpha, InRenderTargetSize)
+		, ParentExtension(InParentExtension)
+		, ViewId(InView.GetViewKey())
+		, ViewFeatureLevel(InView.GetFeatureLevel())
+	{
+	}
+
+	virtual void OnPreRender(FRDGBuilder& GraphBuilder) override
+	{
+		/**
+		 * Note: We abuse the reflection capture view property in a custom render pass to disable
+		 * primitive alpha holdout during its base pass render. Because the holdout is part of the
+		 * primitive uniform buffer, it cannot easily have true/false states in the same frame.
+		 * We would otherwise need to duplicate the entire primitive so that the main render has
+		 * holdout enabled, while the custom render pass primitive has it disabled.
+		 * 
+		 * Maintaining these duplicates becomes problematic, especially with different types of meshes.
+		 * As a result, we consider this reflection capture property override preferable.
+		*/
+		for (FViewInfo* ViewInfo : Views)
+		{
+			// Holdout is ignored during reflections captures to preserve indirect light (see PRIMITIVE_SCENE_DATA_FLAG_HOLDOUT in BasePassPixelShader.usf)
+			ViewInfo->CachedViewUniformShaderParameters->RenderingReflectionCaptureMask = 1.0f;
+			ViewInfo->ViewUniformBuffer.UpdateUniformBufferImmediate(GraphBuilder.RHICmdList, *ViewInfo->CachedViewUniformShaderParameters);
+		}
+
+		const FRDGTextureDesc TextureDesc = FRDGTextureDesc::Create2D(RenderTargetSize, PF_FloatRGBA, FClearValueBinding::Black, TexCreate_RenderTargetable | TexCreate_ShaderResource);
+		RenderTargetTexture = GraphBuilder.CreateTexture(TextureDesc, TEXT("HoldoutCompositeCustomTexture"));
+		AddClearRenderTargetPass(GraphBuilder, RenderTargetTexture, FLinearColor::Black, FIntRect(FInt32Point(), RenderTargetSize));
+	}
+
+	virtual void OnPostRender(FRDGBuilder& GraphBuilder) override
+	{
+		ParentExtension->CollectCustomRenderTarget(ViewId, CreateDilatedTexture(GraphBuilder));
+	}
+
+private:
+	
+	FRDGTextureRef CreateDilatedTexture(FRDGBuilder& GraphBuilder)
+	{
+		RDG_EVENT_SCOPE(GraphBuilder, "HoldoutComposite.Dilate");
+		RDG_GPU_STAT_SCOPE(GraphBuilder, FHoldoutCompositeDilate);
+
+		FGlobalShaderMap* GlobalShaderMap = GetGlobalShaderMap(ViewFeatureLevel);
+
+		const FIntPoint TextureSize = FIntPoint(RenderTargetTexture->Desc.GetSize().X, RenderTargetTexture->Desc.GetSize().Y);
+		const FRDGTextureDesc TextureDesc = FRDGTextureDesc::Create2D(TextureSize, PF_FloatRGBA, FClearValueBinding::Black, TexCreate_UAV | TexCreate_ShaderResource);
+		FRDGTextureRef DilatedTexture = GraphBuilder.CreateTexture(TextureDesc, TEXT("HoldoutCompositeDilatedTexture"));
+
+		// Async compute dilation pass
+		{
+			FDilateHoldoutCompositeCS::FParameters* PassParameters = GraphBuilder.AllocParameters<FDilateHoldoutCompositeCS::FParameters>();
+			PassParameters->InputTexture = RenderTargetTexture;
+			PassParameters->RWOutputTexture = GraphBuilder.CreateUAV(DilatedTexture);
+			PassParameters->Dimensions = TextureSize;
+
+			TShaderMapRef<FDilateHoldoutCompositeCS> ComputeShader(GlobalShaderMap);
+			FComputeShaderUtils::AddPass(
+				GraphBuilder,
+				RDG_EVENT_NAME("HoldoutComposite.Dilate (%dx%d)", TextureSize.X, TextureSize.Y),
+				GSupportsEfficientAsyncCompute ? ERDGPassFlags::AsyncCompute : ERDGPassFlags::Compute,
+				ComputeShader,
+				PassParameters,
+				FComputeShaderUtils::GetGroupCount(TextureSize, FDilateHoldoutCompositeCS::ThreadGroupSize)
+			);
+		}
+
+		return DilatedTexture;
+	}
+
+	FHoldoutCompositeSceneViewExtension* ParentExtension;
+	uint32 ViewId;
+	ERHIFeatureLevel::Type ViewFeatureLevel;
+};
+
+FHoldoutCompositeSceneViewExtension::FHoldoutCompositeSceneViewExtension(const FAutoRegister& AutoReg, UWorld* InWorld)
+	: FWorldSceneViewExtension(AutoReg, InWorld)
+{
+}
+
+FHoldoutCompositeSceneViewExtension::~FHoldoutCompositeSceneViewExtension()
+{
+}
+
+void FHoldoutCompositeSceneViewExtension::RegisterPrimitives(TArrayView<TSoftObjectPtr<UPrimitiveComponent>> InPrimitiveComponents, bool bInHoldoutState)
+{
+	check(IsInGameThread());
+
+	for (const TSoftObjectPtr<UPrimitiveComponent>& InPrimitiveComponent : InPrimitiveComponents)
+	{
+		if (!InPrimitiveComponent.IsValid())
+		{
+			continue;
+		}
+
+		if (!CompositePrimitives.Contains(InPrimitiveComponent))
+		{
+			CompositePrimitives.Add(InPrimitiveComponent);
+			InPrimitiveComponent->SetHoldout(bInHoldoutState);
+		}
+	}
+}
+
+void FHoldoutCompositeSceneViewExtension::UnregisterPrimitives(TArrayView<TSoftObjectPtr<UPrimitiveComponent>> InPrimitiveComponents, bool bInHoldoutState)
+{
+	check(IsInGameThread());
+
+	for (const TSoftObjectPtr<UPrimitiveComponent>& InPrimitiveComponent : InPrimitiveComponents)
+	{
+		if (!InPrimitiveComponent.IsValid())
+		{
+			continue;
+		}
+
+		if (CompositePrimitives.Contains(InPrimitiveComponent))
+		{
+			CompositePrimitives.Remove(InPrimitiveComponent);
+			InPrimitiveComponent->SetHoldout(bInHoldoutState);
+		}
+	}
+
+}
+
+void FHoldoutCompositeSceneViewExtension::CollectCustomRenderTarget(uint32 InViewId, FRDGTextureRef InRenderTarget)
+{
+	CustomRenderTargetPerView_RenderThread.Add(InViewId, InRenderTarget);
+}
+
+bool FHoldoutCompositeSceneViewExtension::IsActiveThisFrame_Internal(const FSceneViewExtensionContext& Context) const
+{
+	bool bIsActive = FWorldSceneViewExtension::IsActiveThisFrame_Internal(Context);
+	bIsActive &= !CompositePrimitives.IsEmpty();
+	bIsActive &= !IsHDREnabled();
+
+	return bIsActive;
+}
+
+
+//~ Begin ISceneViewExtension Interface
+
+int32 FHoldoutCompositeSceneViewExtension::GetPriority() const
+{
+	constexpr int32 OPENCOLORIO_SCENE_VIEW_EXTENSION_PRIORITY = 100;
+	return OPENCOLORIO_SCENE_VIEW_EXTENSION_PRIORITY - 1;
+}
+
+void FHoldoutCompositeSceneViewExtension::SetupViewFamily(FSceneViewFamily& InViewFamily)
+{
+	const UHoldoutCompositeSettings* Settings = GetDefault<UHoldoutCompositeSettings>();
+	if (Settings != nullptr)
+	{
+		bCompositeWithGlobalExposure = Settings->bEnableGlobalExposureComposite;
+	}
+
+	/**
+	 * By default, we use linear with tone curve to match the regular render.
+	 * This also ensures the tonemap post-processing pass preserves alpha precision.
+	 * The usual sRGB encoding is instead applied at the end of the composite.
+	 **/ 
+	InViewFamily.SceneCaptureSource = SCS_FinalToneCurveHDR;
+}
+
+void FHoldoutCompositeSceneViewExtension::SetupView(FSceneViewFamily& InViewFamily, FSceneView& InView)
+{
+	const TWeakObjectPtr<UWorld> WorldPtr = GetWorld();
+	check(WorldPtr.IsValid());
+
+	TSet<FPrimitiveComponentId> HoldoutCompositePrimitiveIds;
+	for (const TSoftObjectPtr<UPrimitiveComponent>& PrimitivePtr : CompositePrimitives)
+	{
+		if (PrimitivePtr.IsValid())
+		{
+			HoldoutCompositePrimitiveIds.Add(PrimitivePtr->GetPrimitiveSceneId());
+		}
+	}
+
+	// Extract the custom render target size
+	FIntPoint RenderTargetViewSize;
+	if (InView.PrimaryScreenPercentageMethod == EPrimaryScreenPercentageMethod::TemporalUpscale)
+	{
+		// Note: This is equivalent to FViewInfo::GetSecondaryViewRectSize
+		FIntRect OutputRect = FIntRect(0, 0,
+			FMath::CeilToInt(InView.UnscaledViewRect.Width() * InViewFamily.SecondaryViewFraction),
+			FMath::CeilToInt(InView.UnscaledViewRect.Height() * InViewFamily.SecondaryViewFraction)
+		);
+		QuantizeSceneBufferSize(OutputRect.Max, RenderTargetViewSize);
+	}
+	else
+	{
+		RenderTargetViewSize = InView.UnscaledViewRect.Size();
+	}
+
+	// Create a new custom render pass to render the composite primitive(s)
+	FHoldoutCompositeCustomRenderPass* CustomRenderPass = new FHoldoutCompositeCustomRenderPass(
+		RenderTargetViewSize,
+		this,
+		InView
+	);
+
+	FSceneInterface::FCustomRenderPassRendererInput PassInput;
+	PassInput.ViewLocation = InView.ViewLocation;
+	PassInput.ViewRotationMatrix = InView.ViewMatrices.GetViewMatrix().RemoveTranslation();
+	PassInput.ViewRotationMatrix.RemoveScaling();
+
+	// Note: Projection matrix here is without jitter, GetProjectionNoAAMatrix() is invalid (not yet available).
+	PassInput.ProjectionMatrix = InView.ViewMatrices.GetProjectionMatrix();
+	PassInput.ViewActor = InView.ViewActor;
+	PassInput.ShowOnlyPrimitives = HoldoutCompositePrimitiveIds;
+	PassInput.CustomRenderPass = CustomRenderPass;
+	PassInput.bIsSceneCapture = true;
+
+	WorldPtr.Get()->Scene->AddCustomRenderPass(&InViewFamily, PassInput);
+}
+
+void FHoldoutCompositeSceneViewExtension::SubscribeToPostProcessingPass(EPostProcessingPass PassId, FAfterPassCallbackDelegateArray& InOutPassCallbacks, bool bIsPassEnabled)
+{
+	if (PassId == EPostProcessingPass::Tonemap)
+	{
+		InOutPassCallbacks.Add(FAfterPassCallbackDelegate::CreateRaw(this, &FHoldoutCompositeSceneViewExtension::PostProcessPassAfterTonemap_RenderThread));
+	}
+}
+
+FScreenPassTexture FHoldoutCompositeSceneViewExtension::PostProcessPassAfterTonemap_RenderThread(FRDGBuilder& GraphBuilder, const FSceneView& InView, const FPostProcessMaterialInputs& Inputs)
+{
+	RDG_EVENT_SCOPE(GraphBuilder, "HoldoutComposite.Final");
+	RDG_GPU_STAT_SCOPE(GraphBuilder, FHoldoutCompositeFinal);
+
+	FScreenPassTexture SceneColor = FScreenPassTexture::CopyFromSlice(GraphBuilder, Inputs.GetInput(EPostProcessMaterialInput::SceneColor));
+	check(SceneColor.IsValid());
+
+	const FRDGTextureRef* CustomRenderPassTexturePtr = CustomRenderTargetPerView_RenderThread.Find(InView.GetViewKey());
+	FRDGTextureRef CustomRenderPassTexture = GSystemTextures.GetBlackAlphaOneDummy(GraphBuilder);
+	
+	if (CustomRenderPassTexturePtr && HasBeenProduced(*CustomRenderPassTexturePtr))
+	{
+		CustomRenderPassTexture = *CustomRenderPassTexturePtr;
+	}
+
+	FGlobalShaderMap* GlobalShaderMap = GetGlobalShaderMap(InView.GetFeatureLevel());
+	const FSceneViewFamily* Family = InView.Family;
+
+	HoldoutComposite::ESceneColorSourceEncoding SourceEncoding = HoldoutComposite::ESceneColorSourceEncoding::Linear;
+	if ((Family->EngineShowFlags.Tonemapper == 0) || (Family->EngineShowFlags.PostProcessing == 0))
+	{
+		SourceEncoding = HoldoutComposite::ESceneColorSourceEncoding::Gamma;
+	}
+	else if (Family->SceneCaptureSource == SCS_FinalColorLDR)
+	{
+		SourceEncoding = HoldoutComposite::ESceneColorSourceEncoding::sRGB;
+	}
+
+	FScreenPassRenderTarget Output = Inputs.OverrideOutput;
+	if (!Output.IsValid())
+	{
+		Output = FScreenPassRenderTarget::CreateFromInput(GraphBuilder, SceneColor, InView.GetOverwriteLoadAction(), TEXT("HoldoutCompositePassOutput"));
+	}
+
+	// Compositing pass
+	{
+		FCompositeHoldoutCompositePS::FPermutationDomain PermutationVector;
+		PermutationVector.Set<FCompositeHoldoutCompositePS::FApplyGlobalExposure>(bCompositeWithGlobalExposure.load());
+
+		FRDGBufferRef EyeAdaptationBuffer = GetEyeAdaptationBuffer(GraphBuilder, InView);
+
+		FCompositeHoldoutCompositePS::FParameters* PassParameters = GraphBuilder.AllocParameters<FCompositeHoldoutCompositePS::FParameters>();
+		PassParameters->View = InView.ViewUniformBuffer;
+		PassParameters->Input = GetScreenPassTextureViewportParameters(FScreenPassTextureViewport(SceneColor));
+		PassParameters->Custom = GetScreenPassTextureViewportParameters(FScreenPassTextureViewport(CustomRenderPassTexture));
+		PassParameters->Output = GetScreenPassTextureViewportParameters(FScreenPassTextureViewport(Output));
+		PassParameters->InputTexture = SceneColor.Texture;
+		PassParameters->InputSampler = TStaticSamplerState<SF_Point>::GetRHI();
+		PassParameters->CustomTexture = CustomRenderPassTexture;
+		PassParameters->CustomSampler = TStaticSamplerState<SF_Bilinear>::GetRHI();
+		PassParameters->EyeAdaptationBuffer = GraphBuilder.CreateSRV(EyeAdaptationBuffer);
+		PassParameters->SourceEncoding = static_cast<uint32>(SourceEncoding);
+		PassParameters->DisplayGamma = Family->RenderTarget->GetDisplayGamma();
+		PassParameters->RenderTargets[0] = Output.GetRenderTargetBinding();
+
+		TShaderMapRef<FCompositeHoldoutCompositePS> PixelShader(GlobalShaderMap, PermutationVector);
+		FPixelShaderUtils::AddFullscreenPass(
+			GraphBuilder,
+			GlobalShaderMap,
+			RDG_EVENT_NAME("HoldoutComposite.Final (%dx%d) PS",
+				Output.ViewRect.Width(), Output.ViewRect.Height()),
+			PixelShader,
+			PassParameters,
+			Output.ViewRect
+		);
+	}
+
+	return MoveTemp(Output);
+}
+
+void FHoldoutCompositeSceneViewExtension::PostRenderViewFamily_RenderThread(FRDGBuilder& GraphBuilder, FSceneViewFamily& InViewFamily)
+{
+	CustomRenderTargetPerView_RenderThread.Reset();
+
+	// Cleanup invalid primitives.
+	for (auto Iter = CompositePrimitives.CreateIterator(); Iter; ++Iter)
+	{
+		if (!Iter->IsValid())
+		{
+			Iter.RemoveCurrent();
+		}
+	}
+}
+
diff --git a/Engine/Plugins/Experimental/Compositing/HoldoutComposite/Source/HoldoutComposite/Private/HoldoutCompositeSceneViewExtension.h b/Engine/Plugins/Experimental/Compositing/HoldoutComposite/Source/HoldoutComposite/Private/HoldoutCompositeSceneViewExtension.h
new file mode 100644
--- /dev/null
+++ b/Engine/Plugins/Experimental/Compositing/HoldoutComposite/Source/HoldoutComposite/Private/HoldoutCompositeSceneViewExtension.h
@@ -0,0 +1,53 @@
+// Copyright Epic Games, Inc. All Rights Reserved.
+
+#pragma once
+
+#include "SceneViewExtension.h"
+#include "RenderGraphFwd.h"
+#include "RendererInterface.h"
+#include "Containers/ContainersFwd.h"
+
+class FViewInfo;
+
+class FHoldoutCompositeSceneViewExtension : public FWorldSceneViewExtension
+{
+public:
+	FHoldoutCompositeSceneViewExtension(const FAutoRegister& AutoReg, UWorld* InWorld);
+	~FHoldoutCompositeSceneViewExtension();
+
+	/* Register primitives for compositing. */
+	void RegisterPrimitives(TArrayView<TSoftObjectPtr<UPrimitiveComponent>> InPrimitiveComponents, bool bInHoldoutState=true);
+
+	/* Unregister primitives for compositing. */
+	void UnregisterPrimitives(TArrayView<TSoftObjectPtr<UPrimitiveComponent>> InPrimitiveComponents, bool bInHoldoutState=false);
+
+	/* Called by the custom render pass to store its view render target for this frame. */
+	void CollectCustomRenderTarget(uint32 InViewId, FRDGTextureRef InRenderTarget);
+
+	//~ Begin ISceneViewExtension Interface
+	virtual int32 GetPriority() const override;
+	virtual void SetupViewFamily(FSceneViewFamily& InViewFamily) override;
+	virtual void SetupView(FSceneViewFamily& InViewFamily, FSceneView& InView) override;
+	virtual void BeginRenderViewFamily(FSceneViewFamily& InViewFamily) override {};
+	virtual void SubscribeToPostProcessingPass(EPostProcessingPass PassId, FAfterPassCallbackDelegateArray& InOutPassCallbacks, bool bIsPassEnabled) override;
+	virtual void PostRenderViewFamily_RenderThread(FRDGBuilder& GraphBuilder, FSceneViewFamily& InViewFamily) override;
+	//~ End ISceneViewExtension Interface
+
+protected:
+	//~ Begin ISceneViewExtension Interface
+	virtual bool IsActiveThisFrame_Internal(const FSceneViewExtensionContext& Context) const override;
+	//~ End ISceneViewExtension Interface
+
+private:
+	FScreenPassTexture PostProcessPassAfterTonemap_RenderThread(FRDGBuilder& GraphBuilder, const FSceneView& View, const FPostProcessMaterialInputs& Inputs);
+
+	// Collection of primitives to render as a custom render pass and composite after post-processing.
+	TSet<TSoftObjectPtr<UPrimitiveComponent>> CompositePrimitives;
+
+	// Custom render pass render targets for each active view
+	TMap<uint32, FRDGTextureRef> CustomRenderTargetPerView_RenderThread;
+
+	// Flat to enable global exposure on the composited render
+	std::atomic_bool bCompositeWithGlobalExposure = false;
+};
+
diff --git a/Engine/Plugins/Experimental/Compositing/HoldoutComposite/Source/HoldoutComposite/Private/HoldoutCompositeSettings.cpp b/Engine/Plugins/Experimental/Compositing/HoldoutComposite/Source/HoldoutComposite/Private/HoldoutCompositeSettings.cpp
new file mode 100644
--- /dev/null
+++ b/Engine/Plugins/Experimental/Compositing/HoldoutComposite/Source/HoldoutComposite/Private/HoldoutCompositeSettings.cpp
@@ -0,0 +1,26 @@
+// Copyright Epic Games, Inc. All Rights Reserved.
+
+#include "HoldoutCompositeSettings.h"
+
+#include UE_INLINE_GENERATED_CPP_BY_NAME(HoldoutCompositeSettings)
+
+UHoldoutCompositeSettings::UHoldoutCompositeSettings()
+	: bEnableGlobalExposureComposite(false)
+{ }
+
+FName UHoldoutCompositeSettings::GetCategoryName() const
+{
+	return TEXT("Plugins");
+}
+
+#if WITH_EDITOR
+FText UHoldoutCompositeSettings::GetSectionText() const
+{
+	return NSLOCTEXT("HoldoutCompositeSettings", "HoldoutCompositeSettingsSection", "Holdout Composite");
+}
+
+FName UHoldoutCompositeSettings::GetSectionName() const
+{
+	return TEXT("Holdout Composite");
+}
+#endif
diff --git a/Engine/Plugins/Experimental/Compositing/HoldoutComposite/Source/HoldoutComposite/Private/HoldoutCompositeSubsystem.cpp b/Engine/Plugins/Experimental/Compositing/HoldoutComposite/Source/HoldoutComposite/Private/HoldoutCompositeSubsystem.cpp
new file mode 100644
--- /dev/null
+++ b/Engine/Plugins/Experimental/Compositing/HoldoutComposite/Source/HoldoutComposite/Private/HoldoutCompositeSubsystem.cpp
@@ -0,0 +1,69 @@
+// Copyright Epic Games, Inc. All Rights Reserved.
+
+#include "HoldoutCompositeSubsystem.h"
+#include "HoldoutCompositeModule.h"
+#include "HoldoutCompositeSceneViewExtension.h"
+
+UHoldoutCompositeSubsystem::UHoldoutCompositeSubsystem()
+{
+}
+
+void UHoldoutCompositeSubsystem::Initialize(FSubsystemCollectionBase& Collection)
+{
+	Super::Initialize(Collection);
+
+	UWorld* World = GetWorld();
+	check(IsValid(World));
+
+	HoldoutCompositeViewExtension = FSceneViewExtensions::NewExtension<FHoldoutCompositeSceneViewExtension>(World);
+}
+
+void UHoldoutCompositeSubsystem::Deinitialize()
+{
+	HoldoutCompositeViewExtension.Reset();
+	
+	Super::Deinitialize();
+}
+
+void UHoldoutCompositeSubsystem::RegisterPrimitive(TSoftObjectPtr<UPrimitiveComponent> InPrimitiveComponent, bool bInHoldoutState)
+{
+	TArray<TSoftObjectPtr<UPrimitiveComponent>> PrimitiveComponents = { InPrimitiveComponent };
+
+	RegisterPrimitives(PrimitiveComponents, bInHoldoutState);
+}
+
+void UHoldoutCompositeSubsystem::RegisterPrimitives(TArrayView<TSoftObjectPtr<UPrimitiveComponent>> InPrimitiveComponents, bool bInHoldoutState)
+{
+	// The compositing relies on alpha preserved through the "tonemapper" post-processing step.
+	static IConsoleVariable* CVarPropagateAlpha = IConsoleManager::Get().FindConsoleVariable(TEXT("r.PostProcessing.PropagateAlpha"));
+	if (CVarPropagateAlpha && CVarPropagateAlpha->GetInt() != 2 /*EAlphaChannelMode::AllowThroughTonemapper*/)
+	{
+		UE_CALL_ONCE([]()
+			{
+				UE_LOG(LogHoldoutComposite, Warning, TEXT("Holdout composite is disabled until r.PostProcessing.PropagateAlpha=2 is set."));
+			}
+		);
+		return;
+	}
+
+	if (HoldoutCompositeViewExtension.IsValid())
+	{
+		HoldoutCompositeViewExtension->RegisterPrimitives(MoveTemp(InPrimitiveComponents), bInHoldoutState);
+	}
+}
+
+void UHoldoutCompositeSubsystem::UnregisterPrimitive(TSoftObjectPtr<UPrimitiveComponent> InPrimitiveComponent, bool bInHoldoutState)
+{
+	TArray<TSoftObjectPtr<UPrimitiveComponent>> PrimitiveComponents = { InPrimitiveComponent };
+	
+	UnregisterPrimitives(PrimitiveComponents, bInHoldoutState);
+}
+
+void UHoldoutCompositeSubsystem::UnregisterPrimitives(TArrayView<TSoftObjectPtr<UPrimitiveComponent>> InPrimitiveComponents, bool bInHoldoutState)
+{
+	if (HoldoutCompositeViewExtension.IsValid())
+	{
+		HoldoutCompositeViewExtension->UnregisterPrimitives(MoveTemp(InPrimitiveComponents), bInHoldoutState);
+	}
+}
+
diff --git a/Engine/Plugins/Experimental/Compositing/HoldoutComposite/Source/HoldoutComposite/Public/HoldoutCompositeComponent.h b/Engine/Plugins/Experimental/Compositing/HoldoutComposite/Source/HoldoutComposite/Public/HoldoutCompositeComponent.h
new file mode 100644
--- /dev/null
+++ b/Engine/Plugins/Experimental/Compositing/HoldoutComposite/Source/HoldoutComposite/Public/HoldoutCompositeComponent.h
@@ -0,0 +1,38 @@
+// Copyright Epic Games, Inc. All Rights Reserved.
+
+#pragma once
+
+#include "CoreTypes.h"
+#include "Components/SceneComponent.h"
+
+#include "HoldoutCompositeComponent.generated.h"
+
+UCLASS(ClassGroup = Rendering, editinlinenew, meta = (BlueprintSpawnableComponent))
+class UHoldoutCompositeComponent : public USceneComponent
+{
+	GENERATED_UCLASS_BODY()
+
+public:
+	//~ Begin UActorComponent interface
+	HOLDOUTCOMPOSITE_API virtual void OnRegister() override;
+	HOLDOUTCOMPOSITE_API virtual void OnUnregister() override;
+	//~ End UActorComponent interface
+	
+	//~ Begin USceneComponent Interface.
+	HOLDOUTCOMPOSITE_API virtual void OnAttachmentChanged() override;
+	HOLDOUTCOMPOSITE_API virtual void DetachFromComponent(const FDetachmentTransformRules& DetachmentRules) override;
+	//~ End USceneComponent Interface
+
+private:
+
+	/* Private implementation of the register method. */
+	void RegisterCompositeImpl();
+
+	/* Private implementation of the unregister method. */
+	void UnregisterCompositeImpl();
+
+	/* Primitive holdout property value of the parent component. */
+	UPROPERTY()
+	TOptional<bool> bCachedParentHoldout = false;
+};
+
diff --git a/Engine/Plugins/Experimental/Compositing/HoldoutComposite/Source/HoldoutComposite/Public/HoldoutCompositeModule.h b/Engine/Plugins/Experimental/Compositing/HoldoutComposite/Source/HoldoutComposite/Public/HoldoutCompositeModule.h
new file mode 100644
--- /dev/null
+++ b/Engine/Plugins/Experimental/Compositing/HoldoutComposite/Source/HoldoutComposite/Public/HoldoutCompositeModule.h
@@ -0,0 +1,18 @@
+// Copyright Epic Games, Inc. All Rights Reserved.
+
+#pragma once
+
+#include "Modules/ModuleInterface.h"
+#include "Logging/LogMacros.h"
+
+HOLDOUTCOMPOSITE_API DECLARE_LOG_CATEGORY_EXTERN(LogHoldoutComposite, Log, All);
+
+class FHoldoutCompositeModule : public IModuleInterface
+{
+public:
+	// IModuleInterface interface
+	virtual void StartupModule() override;
+	virtual void ShutdownModule() override;
+	// End of IModuleInterface interface
+};
+
diff --git a/Engine/Plugins/Experimental/Compositing/HoldoutComposite/Source/HoldoutComposite/Public/HoldoutCompositeSettings.h b/Engine/Plugins/Experimental/Compositing/HoldoutComposite/Source/HoldoutComposite/Public/HoldoutCompositeSettings.h
new file mode 100644
--- /dev/null
+++ b/Engine/Plugins/Experimental/Compositing/HoldoutComposite/Source/HoldoutComposite/Public/HoldoutCompositeSettings.h
@@ -0,0 +1,36 @@
+// Copyright Epic Games, Inc. All Rights Reserved.
+
+#pragma once
+
+#include "CoreTypes.h"
+#include "UObject/Object.h"
+#include "UObject/ObjectMacros.h"
+#include "Engine/DeveloperSettings.h"
+
+#include "HoldoutCompositeSettings.generated.h"
+
+/**
+ * Settings for the HoldoutComposite module.
+ */
+UCLASS(config = Engine)
+class HOLDOUTCOMPOSITE_API UHoldoutCompositeSettings : public UDeveloperSettings
+{
+	GENERATED_BODY()
+
+public:
+	UHoldoutCompositeSettings();
+
+	//~ Begin UDeveloperSettings interface
+	virtual FName GetCategoryName() const;
+#if WITH_EDITOR
+	virtual FText GetSectionText() const override;
+	virtual FName GetSectionName() const override;
+#endif
+	//~ End UDeveloperSettings interface
+
+public:
+	/** When enabled, the view global exposure is applied onto the separate render when composited. */
+	UPROPERTY(config, EditAnywhere, Category = General)
+	bool bEnableGlobalExposureComposite;
+};
+
diff --git a/Engine/Plugins/Experimental/Compositing/HoldoutComposite/Source/HoldoutComposite/Public/HoldoutCompositeSubsystem.h b/Engine/Plugins/Experimental/Compositing/HoldoutComposite/Source/HoldoutComposite/Public/HoldoutCompositeSubsystem.h
new file mode 100644
--- /dev/null
+++ b/Engine/Plugins/Experimental/Compositing/HoldoutComposite/Source/HoldoutComposite/Public/HoldoutCompositeSubsystem.h
@@ -0,0 +1,46 @@
+// Copyright Epic Games, Inc. All Rights Reserved.
+
+#pragma once
+
+#include "MaterialShared.h"
+#include "Subsystems/WorldSubsystem.h"
+#include "UObject/SoftObjectPtr.h"
+
+#include "HoldoutCompositeSubsystem.generated.h"
+
+class FHoldoutCompositeSceneViewExtension;
+
+/**
+ * Composite subsytem used as an interface to the (private) scene view extension.
+ */
+UCLASS(BlueprintType, Transient)
+class UHoldoutCompositeSubsystem : public UWorldSubsystem
+{
+	GENERATED_BODY()
+
+public:
+	UHoldoutCompositeSubsystem();
+
+	// USubsystem implementation Begin
+	virtual void Initialize(FSubsystemCollectionBase& Collection) override;
+	virtual void Deinitialize() override;
+	// USubsystem implementation End
+
+	/* Register a single primitive for compositing. */
+	UFUNCTION(BlueprintCallable, Category = "Holdout Composite")
+	HOLDOUTCOMPOSITE_API void RegisterPrimitive(TSoftObjectPtr<UPrimitiveComponent> InPrimitiveComponent, bool bInHoldoutState=true);
+
+	/* Register multiple primitives for compositing. */
+	HOLDOUTCOMPOSITE_API void RegisterPrimitives(TArrayView<TSoftObjectPtr<UPrimitiveComponent>> InPrimitiveComponents, bool bInHoldoutState=true);
+
+	/* Unregister a single primitive from compositing. */
+	UFUNCTION(BlueprintCallable, Category = "Holdout Composite")
+	HOLDOUTCOMPOSITE_API void UnregisterPrimitive(TSoftObjectPtr<UPrimitiveComponent> InPrimitiveComponent, bool bInHoldoutState=false);
+
+	/* Unregister multiple primitives from compositing. */
+	HOLDOUTCOMPOSITE_API void UnregisterPrimitives(TArrayView<TSoftObjectPtr<UPrimitiveComponent>> InPrimitiveComponents, bool bInHoldoutState=false);
+
+private:
+	TSharedPtr<FHoldoutCompositeSceneViewExtension, ESPMode::ThreadSafe> HoldoutCompositeViewExtension;
+};
+
