HoldoutComposite: Partial fix to prohibited private includes.

#rnx

diff --git a/Engine/Plugins/Experimental/Compositing/HoldoutComposite/Source/HoldoutComposite/Private/HoldoutCompositeSceneViewExtension.cpp b/Engine/Plugins/Experimental/Compositing/HoldoutComposite/Source/HoldoutComposite/Private/HoldoutCompositeSceneViewExtension.cpp
--- a/Engine/Plugins/Experimental/Compositing/HoldoutComposite/Source/HoldoutComposite/Private/HoldoutCompositeSceneViewExtension.cpp
+++ b/Engine/Plugins/Experimental/Compositing/HoldoutComposite/Source/HoldoutComposite/Private/HoldoutCompositeSceneViewExtension.cpp
@@ -22,7 +22,6 @@
 
 // Private renderer includes currently needed...
 #include "SceneRendering.h"
-#include "SceneTextureParameters.h"
 
 DECLARE_GPU_STAT_NAMED(FHoldoutCompositeDilate, TEXT("HoldoutComposite.Dilate"));
 DECLARE_GPU_STAT_NAMED(FHoldoutCompositeFinal, TEXT("HoldoutComposite.Final"));
@@ -353,7 +352,7 @@ FScreenPassTexture FHoldoutCompositeSceneViewExtension::PostProcessPassAfterTone
 		FCompositeHoldoutCompositePS::FPermutationDomain PermutationVector;
 		PermutationVector.Set<FCompositeHoldoutCompositePS::FApplyGlobalExposure>(bCompositeWithGlobalExposure.load());
 
-		FRDGBufferRef EyeAdaptationBuffer = GetEyeAdaptationBuffer(GraphBuilder, InView);
+		FRDGBufferRef EyeAdaptationBuffer = GraphBuilder.RegisterExternalBuffer(InView.GetEyeAdaptationBuffer(), ERDGBufferFlags::MultiFrame);
 
 		FCompositeHoldoutCompositePS::FParameters* PassParameters = GraphBuilder.AllocParameters<FCompositeHoldoutCompositePS::FParameters>();
 		PassParameters->View = InView.ViewUniformBuffer;
