HoldoutComposite: Now default to follow scene exposure, and run before OCIO.

#rnx

diff --git a/Engine/Plugins/Experimental/Compositing/HoldoutComposite/Shaders/Private/HoldoutCompositeFinal.usf b/Engine/Plugins/Experimental/Compositing/HoldoutComposite/Shaders/Private/HoldoutCompositeFinal.usf
--- a/Engine/Plugins/Experimental/Compositing/HoldoutComposite/Shaders/Private/HoldoutCompositeFinal.usf
+++ b/Engine/Plugins/Experimental/Compositing/HoldoutComposite/Shaders/Private/HoldoutCompositeFinal.usf
@@ -15,8 +15,8 @@ Texture2D CustomTexture;
 SamplerState CustomSampler;
 
 StructuredBuffer<float4> EyeAdaptationBuffer;
-uint SourceEncoding;
-float DisplayGamma;
+uint2 Encodings;
+float2 DisplayGamma;
 
 void MainPS(
     float4 SvPosition : SV_POSITION,
@@ -37,16 +37,26 @@ void MainPS(
 #endif
 
     BRANCH
-    if(SourceEncoding == 1)
+    if(Encodings.x == 1)
     {
-        SceneColor.rgb = pow(SceneColor.rgb, DisplayGamma);
+        SceneColor.rgb = pow(SceneColor.rgb, DisplayGamma.x);
     }
-    else if(SourceEncoding == 2)
+    else if(Encodings.x == 2)
     {
         SceneColor.rgb = sRGBToLinear(SceneColor.rgb);
     }
 
     // TODO: HDR cases will need to be handled explicitely.
-    OutColor.rgb = LinearToSrgb(CustomColor.rgb * SceneColor.a + SceneColor.rgb);
+    OutColor.rgb = CustomColor.rgb * SceneColor.a + SceneColor.rgb;
     OutColor.a = CustomColor.a * SceneColor.a;
+
+    BRANCH
+    if(Encodings.y == 1)
+    {
+        OutColor.rgb = pow(OutColor.rgb, DisplayGamma.y);
+    }
+    else if(Encodings.y == 2)
+    {
+        OutColor.rgb = LinearToSrgb(OutColor.rgb);
+    }
 }
diff --git a/Engine/Plugins/Experimental/Compositing/HoldoutComposite/Source/HoldoutComposite/Private/HoldoutCompositeSceneViewExtension.cpp b/Engine/Plugins/Experimental/Compositing/HoldoutComposite/Source/HoldoutComposite/Private/HoldoutCompositeSceneViewExtension.cpp
--- a/Engine/Plugins/Experimental/Compositing/HoldoutComposite/Source/HoldoutComposite/Private/HoldoutCompositeSceneViewExtension.cpp
+++ b/Engine/Plugins/Experimental/Compositing/HoldoutComposite/Source/HoldoutComposite/Private/HoldoutCompositeSceneViewExtension.cpp
@@ -76,8 +76,8 @@ class FCompositeHoldoutCompositePS : public FGlobalShader
 		SHADER_PARAMETER_RDG_TEXTURE(Texture2D, CustomTexture)
 		SHADER_PARAMETER_SAMPLER(SamplerState, CustomSampler)
 		SHADER_PARAMETER_RDG_BUFFER_SRV(StructuredBuffer<float4>, EyeAdaptationBuffer)
-		SHADER_PARAMETER(uint32, SourceEncoding)
-		SHADER_PARAMETER(float, DisplayGamma)
+		SHADER_PARAMETER(FUint32Vector2, Encodings)
+		SHADER_PARAMETER(FVector2f, DisplayGamma)
 
 		RENDER_TARGET_BINDING_SLOTS()
 	END_SHADER_PARAMETER_STRUCT()
@@ -232,8 +232,13 @@ bool FHoldoutCompositeSceneViewExtension::IsActiveThisFrame_Internal(const FScen
 
 int32 FHoldoutCompositeSceneViewExtension::GetPriority() const
 {
-	constexpr int32 OPENCOLORIO_SCENE_VIEW_EXTENSION_PRIORITY = 100;
-	return OPENCOLORIO_SCENE_VIEW_EXTENSION_PRIORITY - 1;
+	const UHoldoutCompositeSettings* Settings = GetDefault<UHoldoutCompositeSettings>();
+	if (Settings != nullptr)
+	{
+		return Settings->SceneViewExtensionPriority;
+	}
+
+	return 0;
 }
 
 void FHoldoutCompositeSceneViewExtension::SetupViewFamily(FSceneViewFamily& InViewFamily)
@@ -241,7 +246,7 @@ void FHoldoutCompositeSceneViewExtension::SetupViewFamily(FSceneViewFamily& InVi
 	const UHoldoutCompositeSettings* Settings = GetDefault<UHoldoutCompositeSettings>();
 	if (Settings != nullptr)
 	{
-		bCompositeWithGlobalExposure = Settings->bEnableGlobalExposureComposite;
+		bCompositeFollowsSceneExposure = Settings->bCompositeFollowsSceneExposure;
 	}
 
 	/**
@@ -249,6 +254,7 @@ void FHoldoutCompositeSceneViewExtension::SetupViewFamily(FSceneViewFamily& InVi
 	 * This also ensures the tonemap post-processing pass preserves alpha precision.
 	 * The usual sRGB encoding is instead applied at the end of the composite.
 	 **/ 
+	OriginalSceneCaptureSource = InViewFamily.SceneCaptureSource;
 	InViewFamily.SceneCaptureSource = SCS_FinalToneCurveHDR;
 }
 
@@ -316,6 +322,7 @@ FScreenPassTexture FHoldoutCompositeSceneViewExtension::PostProcessPassAfterTone
 {
 	RDG_EVENT_SCOPE(GraphBuilder, "HoldoutComposite.Final");
 	RDG_GPU_STAT_SCOPE(GraphBuilder, FHoldoutCompositeFinal);
+	using namespace HoldoutComposite;
 
 	FScreenPassTexture SceneColor = FScreenPassTexture::CopyFromSlice(GraphBuilder, Inputs.GetInput(EPostProcessMaterialInput::SceneColor));
 	check(SceneColor.IsValid());
@@ -331,14 +338,22 @@ FScreenPassTexture FHoldoutCompositeSceneViewExtension::PostProcessPassAfterTone
 	FGlobalShaderMap* GlobalShaderMap = GetGlobalShaderMap(InView.GetFeatureLevel());
 	const FSceneViewFamily* Family = InView.Family;
 
-	HoldoutComposite::ESceneColorSourceEncoding SourceEncoding = HoldoutComposite::ESceneColorSourceEncoding::Linear;
+	FUint32Vector2 Encodings = FUint32Vector2::ZeroValue; // ESceneColorSourceEncoding::Linear
 	if ((Family->EngineShowFlags.Tonemapper == 0) || (Family->EngineShowFlags.PostProcessing == 0))
 	{
-		SourceEncoding = HoldoutComposite::ESceneColorSourceEncoding::Gamma;
+		Encodings.X = static_cast<uint32>(ESceneColorSourceEncoding::Gamma);
+		Encodings.Y = static_cast<uint32>(ESceneColorSourceEncoding::Gamma);
 	}
 	else if (Family->SceneCaptureSource == SCS_FinalColorLDR)
 	{
-		SourceEncoding = HoldoutComposite::ESceneColorSourceEncoding::sRGB;
+		Encodings.X = static_cast<uint32>(ESceneColorSourceEncoding::sRGB);
+		Encodings.Y = static_cast<uint32>(ESceneColorSourceEncoding::sRGB);
+	}
+	else if (Family->SceneCaptureSource == SCS_FinalToneCurveHDR
+		&& OriginalSceneCaptureSource != ESceneCaptureSource::SCS_FinalToneCurveHDR)
+	{
+		// Special setup to preserve alpha precision, per earlier note.
+		Encodings.Y = static_cast<uint32>(ESceneColorSourceEncoding::sRGB);
 	}
 
 	FScreenPassRenderTarget Output = Inputs.OverrideOutput;
@@ -350,7 +365,7 @@ FScreenPassTexture FHoldoutCompositeSceneViewExtension::PostProcessPassAfterTone
 	// Compositing pass
 	{
 		FCompositeHoldoutCompositePS::FPermutationDomain PermutationVector;
-		PermutationVector.Set<FCompositeHoldoutCompositePS::FApplyGlobalExposure>(bCompositeWithGlobalExposure.load());
+		PermutationVector.Set<FCompositeHoldoutCompositePS::FApplyGlobalExposure>(bCompositeFollowsSceneExposure.load());
 
 		FRDGBufferRef EyeAdaptationBuffer = GraphBuilder.RegisterExternalBuffer(InView.GetEyeAdaptationBuffer(), ERDGBufferFlags::MultiFrame);
 
@@ -364,8 +379,8 @@ FScreenPassTexture FHoldoutCompositeSceneViewExtension::PostProcessPassAfterTone
 		PassParameters->CustomTexture = CustomRenderPassTexture;
 		PassParameters->CustomSampler = TStaticSamplerState<SF_Bilinear>::GetRHI();
 		PassParameters->EyeAdaptationBuffer = GraphBuilder.CreateSRV(EyeAdaptationBuffer);
-		PassParameters->SourceEncoding = static_cast<uint32>(SourceEncoding);
-		PassParameters->DisplayGamma = Family->RenderTarget->GetDisplayGamma();
+		PassParameters->Encodings = Encodings;
+		PassParameters->DisplayGamma = FVector2f(Family->RenderTarget->GetDisplayGamma(), 1.0f / Family->RenderTarget->GetDisplayGamma());
 		PassParameters->RenderTargets[0] = Output.GetRenderTargetBinding();
 
 		TShaderMapRef<FCompositeHoldoutCompositePS> PixelShader(GlobalShaderMap, PermutationVector);
diff --git a/Engine/Plugins/Experimental/Compositing/HoldoutComposite/Source/HoldoutComposite/Private/HoldoutCompositeSceneViewExtension.h b/Engine/Plugins/Experimental/Compositing/HoldoutComposite/Source/HoldoutComposite/Private/HoldoutCompositeSceneViewExtension.h
--- a/Engine/Plugins/Experimental/Compositing/HoldoutComposite/Source/HoldoutComposite/Private/HoldoutCompositeSceneViewExtension.h
+++ b/Engine/Plugins/Experimental/Compositing/HoldoutComposite/Source/HoldoutComposite/Private/HoldoutCompositeSceneViewExtension.h
@@ -47,7 +47,10 @@ private:
 	// Custom render pass render targets for each active view
 	TMap<uint32, FRDGTextureRef> CustomRenderTargetPerView_RenderThread;
 
-	// Flat to enable global exposure on the composited render
-	std::atomic_bool bCompositeWithGlobalExposure = false;
+	// Flag to enable global exposure on the composited render
+	std::atomic_bool bCompositeFollowsSceneExposure = false;
+
+	// Track the original view family encoding
+	ESceneCaptureSource OriginalSceneCaptureSource = ESceneCaptureSource::SCS_FinalColorLDR;
 };
 
diff --git a/Engine/Plugins/Experimental/Compositing/HoldoutComposite/Source/HoldoutComposite/Private/HoldoutCompositeSettings.cpp b/Engine/Plugins/Experimental/Compositing/HoldoutComposite/Source/HoldoutComposite/Private/HoldoutCompositeSettings.cpp
--- a/Engine/Plugins/Experimental/Compositing/HoldoutComposite/Source/HoldoutComposite/Private/HoldoutCompositeSettings.cpp
+++ b/Engine/Plugins/Experimental/Compositing/HoldoutComposite/Source/HoldoutComposite/Private/HoldoutCompositeSettings.cpp
@@ -5,7 +5,8 @@
 #include UE_INLINE_GENERATED_CPP_BY_NAME(HoldoutCompositeSettings)
 
 UHoldoutCompositeSettings::UHoldoutCompositeSettings()
-	: bEnableGlobalExposureComposite(false)
+	: bCompositeFollowsSceneExposure(true)
+	, SceneViewExtensionPriority(100 + 1 /* OPENCOLORIO_SCENE_VIEW_EXTENSION_PRIORITY + 1 */)
 { }
 
 FName UHoldoutCompositeSettings::GetCategoryName() const
diff --git a/Engine/Plugins/Experimental/Compositing/HoldoutComposite/Source/HoldoutComposite/Public/HoldoutCompositeSettings.h b/Engine/Plugins/Experimental/Compositing/HoldoutComposite/Source/HoldoutComposite/Public/HoldoutCompositeSettings.h
--- a/Engine/Plugins/Experimental/Compositing/HoldoutComposite/Source/HoldoutComposite/Public/HoldoutCompositeSettings.h
+++ b/Engine/Plugins/Experimental/Compositing/HoldoutComposite/Source/HoldoutComposite/Public/HoldoutCompositeSettings.h
@@ -31,6 +31,10 @@ public:
 public:
 	/** When enabled, the view global exposure is applied onto the separate render when composited. */
 	UPROPERTY(config, EditAnywhere, Category = General)
-	bool bEnableGlobalExposureComposite;
+	bool bCompositeFollowsSceneExposure;
+
+	/** Composite (scene view extension) pass priority, which defaults to before OpenColorIO. */
+	UPROPERTY(config, EditAnywhere, Category = General, AdvancedDisplay)
+	int32 SceneViewExtensionPriority;
 };
 
