HoldoutComposite: Improved scene view extension since CL 34277599, and added MinimalAPI to component.

* Avoid post processing callback when there is no work (no custom render pass render target was produced).
* Check for view hidden/show only primitives.

#rnx

diff --git a/Engine/Plugins/Experimental/Compositing/HoldoutComposite/Source/HoldoutComposite/Private/HoldoutCompositeSceneViewExtension.cpp b/Engine/Plugins/Experimental/Compositing/HoldoutComposite/Source/HoldoutComposite/Private/HoldoutCompositeSceneViewExtension.cpp
--- a/Engine/Plugins/Experimental/Compositing/HoldoutComposite/Source/HoldoutComposite/Private/HoldoutCompositeSceneViewExtension.cpp
+++ b/Engine/Plugins/Experimental/Compositing/HoldoutComposite/Source/HoldoutComposite/Private/HoldoutCompositeSceneViewExtension.cpp
@@ -268,10 +268,27 @@ void FHoldoutCompositeSceneViewExtension::SetupView(FSceneViewFamily& InViewFami
 	{
 		if (PrimitivePtr.IsValid())
 		{
-			HoldoutCompositePrimitiveIds.Add(PrimitivePtr->GetPrimitiveSceneId());
+			const FPrimitiveComponentId PrimId = PrimitivePtr->GetPrimitiveSceneId();
+
+			if (InView.ShowOnlyPrimitives.IsSet())
+			{
+				if (InView.ShowOnlyPrimitives.GetValue().Contains(PrimId))
+				{
+					HoldoutCompositePrimitiveIds.Add(PrimitivePtr->GetPrimitiveSceneId());
+				}
+			}
+			else if (!InView.HiddenPrimitives.Contains(PrimId))
+			{
+				HoldoutCompositePrimitiveIds.Add(PrimitivePtr->GetPrimitiveSceneId());
+			}
 		}
 	}
 
+	if (HoldoutCompositePrimitiveIds.IsEmpty())
+	{
+		return;
+	}
+
 	// Extract the custom render target size
 	FIntPoint RenderTargetViewSize;
 	if (InView.PrimaryScreenPercentageMethod == EPrimaryScreenPercentageMethod::TemporalUpscale)
@@ -310,8 +327,14 @@ void FHoldoutCompositeSceneViewExtension::SetupView(FSceneViewFamily& InViewFami
 	WorldPtr.Get()->Scene->AddCustomRenderPass(&InViewFamily, PassInput);
 }
 
-void FHoldoutCompositeSceneViewExtension::SubscribeToPostProcessingPass(EPostProcessingPass PassId, const FSceneView& View, FAfterPassCallbackDelegateArray& InOutPassCallbacks, bool bIsPassEnabled)
+void FHoldoutCompositeSceneViewExtension::SubscribeToPostProcessingPass(EPostProcessingPass PassId, const FSceneView& InView, FAfterPassCallbackDelegateArray& InOutPassCallbacks, bool bIsPassEnabled)
 {
+	if (!CustomRenderTargetPerView_RenderThread.Contains(InView.GetViewKey()))
+	{
+		// Early-out to avoid needless work in the post processing callback(s).
+		return;
+	}
+
 	if (PassId == EPostProcessingPass::Tonemap)
 	{
 		InOutPassCallbacks.Add(FAfterPassCallbackDelegate::CreateRaw(this, &FHoldoutCompositeSceneViewExtension::PostProcessPassAfterTonemap_RenderThread));
diff --git a/Engine/Plugins/Experimental/Compositing/HoldoutComposite/Source/HoldoutComposite/Private/HoldoutCompositeSceneViewExtension.h b/Engine/Plugins/Experimental/Compositing/HoldoutComposite/Source/HoldoutComposite/Private/HoldoutCompositeSceneViewExtension.h
--- a/Engine/Plugins/Experimental/Compositing/HoldoutComposite/Source/HoldoutComposite/Private/HoldoutCompositeSceneViewExtension.h
+++ b/Engine/Plugins/Experimental/Compositing/HoldoutComposite/Source/HoldoutComposite/Private/HoldoutCompositeSceneViewExtension.h
@@ -29,7 +29,7 @@ public:
 	virtual void SetupViewFamily(FSceneViewFamily& InViewFamily) override;
 	virtual void SetupView(FSceneViewFamily& InViewFamily, FSceneView& InView) override;
 	virtual void BeginRenderViewFamily(FSceneViewFamily& InViewFamily) override {};
-	virtual void SubscribeToPostProcessingPass(EPostProcessingPass PassId, const FSceneView& View, FAfterPassCallbackDelegateArray& InOutPassCallbacks, bool bIsPassEnabled) override;
+	virtual void SubscribeToPostProcessingPass(EPostProcessingPass PassId, const FSceneView& InView, FAfterPassCallbackDelegateArray& InOutPassCallbacks, bool bIsPassEnabled) override;
 	virtual void PostRenderViewFamily_RenderThread(FRDGBuilder& GraphBuilder, FSceneViewFamily& InViewFamily) override;
 	//~ End ISceneViewExtension Interface
 
diff --git a/Engine/Plugins/Experimental/Compositing/HoldoutComposite/Source/HoldoutComposite/Public/HoldoutCompositeComponent.h b/Engine/Plugins/Experimental/Compositing/HoldoutComposite/Source/HoldoutComposite/Public/HoldoutCompositeComponent.h
--- a/Engine/Plugins/Experimental/Compositing/HoldoutComposite/Source/HoldoutComposite/Public/HoldoutCompositeComponent.h
+++ b/Engine/Plugins/Experimental/Compositing/HoldoutComposite/Source/HoldoutComposite/Public/HoldoutCompositeComponent.h
@@ -7,7 +7,7 @@
 
 #include "HoldoutCompositeComponent.generated.h"
 
-UCLASS(ClassGroup = Rendering, editinlinenew, meta = (BlueprintSpawnableComponent))
+UCLASS(ClassGroup = Rendering, editinlinenew, meta = (BlueprintSpawnableComponent), MinimalAPI)
 class UHoldoutCompositeComponent : public USceneComponent
 {
 	GENERATED_UCLASS_BODY()
