HoldoutComposite: Removed forced linear render to preserve alpha precision since the issue is now fixed in post-processing (34412772).

#rnx

diff --git a/Engine/Plugins/Experimental/Compositing/HoldoutComposite/Source/HoldoutComposite/Private/HoldoutCompositeSceneViewExtension.cpp b/Engine/Plugins/Experimental/Compositing/HoldoutComposite/Source/HoldoutComposite/Private/HoldoutCompositeSceneViewExtension.cpp
--- a/Engine/Plugins/Experimental/Compositing/HoldoutComposite/Source/HoldoutComposite/Private/HoldoutCompositeSceneViewExtension.cpp
+++ b/Engine/Plugins/Experimental/Compositing/HoldoutComposite/Source/HoldoutComposite/Private/HoldoutCompositeSceneViewExtension.cpp
@@ -248,14 +248,6 @@ void FHoldoutCompositeSceneViewExtension::SetupViewFamily(FSceneViewFamily& InVi
 	{
 		bCompositeFollowsSceneExposure = Settings->bCompositeFollowsSceneExposure;
 	}
-
-	/**
-	 * By default, we use linear with tone curve to match the regular render.
-	 * This also ensures the tonemap post-processing pass preserves alpha precision.
-	 * The usual sRGB encoding is instead applied at the end of the composite.
-	 **/ 
-	OriginalSceneCaptureSource = InViewFamily.SceneCaptureSource;
-	InViewFamily.SceneCaptureSource = SCS_FinalToneCurveHDR;
 }
 
 void FHoldoutCompositeSceneViewExtension::SetupView(FSceneViewFamily& InViewFamily, FSceneView& InView)
@@ -372,12 +364,6 @@ FScreenPassTexture FHoldoutCompositeSceneViewExtension::PostProcessPassAfterTone
 		Encodings.X = static_cast<uint32>(ESceneColorSourceEncoding::sRGB);
 		Encodings.Y = static_cast<uint32>(ESceneColorSourceEncoding::sRGB);
 	}
-	else if (Family->SceneCaptureSource == SCS_FinalToneCurveHDR
-		&& OriginalSceneCaptureSource != ESceneCaptureSource::SCS_FinalToneCurveHDR)
-	{
-		// Special setup to preserve alpha precision, per earlier note.
-		Encodings.Y = static_cast<uint32>(ESceneColorSourceEncoding::sRGB);
-	}
 
 	FScreenPassRenderTarget Output = Inputs.OverrideOutput;
 	if (!Output.IsValid())
diff --git a/Engine/Plugins/Experimental/Compositing/HoldoutComposite/Source/HoldoutComposite/Private/HoldoutCompositeSceneViewExtension.h b/Engine/Plugins/Experimental/Compositing/HoldoutComposite/Source/HoldoutComposite/Private/HoldoutCompositeSceneViewExtension.h
--- a/Engine/Plugins/Experimental/Compositing/HoldoutComposite/Source/HoldoutComposite/Private/HoldoutCompositeSceneViewExtension.h
+++ b/Engine/Plugins/Experimental/Compositing/HoldoutComposite/Source/HoldoutComposite/Private/HoldoutCompositeSceneViewExtension.h
@@ -49,8 +49,5 @@ private:
 
 	// Flag to enable global exposure on the composited render
 	std::atomic_bool bCompositeFollowsSceneExposure = false;
-
-	// Track the original view family encoding
-	ESceneCaptureSource OriginalSceneCaptureSource = ESceneCaptureSource::SCS_FinalColorLDR;
 };
 
