HoldoutComposite: Add support for engine lens distortion.

diff --git a/Engine/Plugins/Experimental/Compositing/HoldoutComposite/Shaders/Private/HoldoutCompositeFinal.usf b/Engine/Plugins/Experimental/Compositing/HoldoutComposite/Shaders/Private/HoldoutCompositeFinal.usf
--- a/Engine/Plugins/Experimental/Compositing/HoldoutComposite/Shaders/Private/HoldoutCompositeFinal.usf
+++ b/Engine/Plugins/Experimental/Compositing/HoldoutComposite/Shaders/Private/HoldoutCompositeFinal.usf
@@ -3,6 +3,7 @@
 #include "/Engine/Private/Common.ush"
 #include "/Engine/Private/GammaCorrectionCommon.ush"
 #include "/Engine/Private/ScreenPass.ush"
+#include "/Engine/Private/LensDistortion.ush"
 
 SCREEN_PASS_TEXTURE_VIEWPORT(Input)
 SCREEN_PASS_TEXTURE_VIEWPORT(Custom)
@@ -14,6 +15,9 @@ SamplerState InputSampler;
 Texture2D CustomTexture;
 SamplerState CustomSampler;
 
+Texture2D<float2> UndistortingDisplacementTexture;
+SamplerState UndistortingDisplacementSampler;
+
 StructuredBuffer<float4> EyeAdaptationBuffer;
 uint2 Encodings;
 float2 DisplayGamma;
@@ -23,12 +27,13 @@ void MainPS(
     out float4 OutColor : SV_Target0
     )
 {
-    float2 ViewportUV = (SvPosition.xy - Output_ViewportMin) * Output_ViewportSizeInverse;
+    float2 DistortedViewportUV = (SvPosition.xy - Output_ViewportMin) * Output_ViewportSizeInverse;
+    const float2 UndistortedViewportUV = ApplyLensDistortionOnViewportUV(UndistortingDisplacementTexture, UndistortingDisplacementSampler, DistortedViewportUV);
     
-    const float2 InputUV = (ViewportUV * Input_ViewportSize + Input_ViewportMin) * Input_ExtentInverse;
+    const float2 InputUV = (DistortedViewportUV * Input_ViewportSize + Input_ViewportMin) * Input_ExtentInverse;
     float4 SceneColor = Texture2DSample(InputTexture, InputSampler, InputUV);
 
-    const float2 CustomUV = (ViewportUV * Custom_ViewportSize + Custom_ViewportMin) * Custom_ExtentInverse;
+    const float2 CustomUV = (UndistortedViewportUV * Custom_ViewportSize + Custom_ViewportMin) * Custom_ExtentInverse;
     float4 CustomColor = Texture2DSample(CustomTexture, CustomSampler, CustomUV);
 
 #if PERMUTATION_GLOBAL_EXPOSURE
diff --git a/Engine/Plugins/Experimental/Compositing/HoldoutComposite/Source/HoldoutComposite/HoldoutComposite.Build.cs b/Engine/Plugins/Experimental/Compositing/HoldoutComposite/Source/HoldoutComposite/HoldoutComposite.Build.cs
--- a/Engine/Plugins/Experimental/Compositing/HoldoutComposite/Source/HoldoutComposite/HoldoutComposite.Build.cs
+++ b/Engine/Plugins/Experimental/Compositing/HoldoutComposite/Source/HoldoutComposite/HoldoutComposite.Build.cs
@@ -6,7 +6,10 @@ public class HoldoutComposite : ModuleRules
 {
 	public HoldoutComposite(ReadOnlyTargetRules Target) : base(Target)
 	{
-		PCHUsage = ModuleRules.PCHUsageMode.UseExplicitOrSharedPCHs;
+		// The HoldoutComposite plugin is distributed with engine hot fixes and thus isn't tied to binary
+		// compatibility between hotfixes by only using Public/ interface of the renderer, but also Internal/ ones.
+		bTreatAsEngineModule = true;
+		PCHUsage = PCHUsageMode.UseExplicitOrSharedPCHs;
 
 		PrivateIncludePaths.AddRange(
 			new string[] {
diff --git a/Engine/Plugins/Experimental/Compositing/HoldoutComposite/Source/HoldoutComposite/Private/HoldoutCompositeSceneViewExtension.cpp b/Engine/Plugins/Experimental/Compositing/HoldoutComposite/Source/HoldoutComposite/Private/HoldoutCompositeSceneViewExtension.cpp
--- a/Engine/Plugins/Experimental/Compositing/HoldoutComposite/Source/HoldoutComposite/Private/HoldoutCompositeSceneViewExtension.cpp
+++ b/Engine/Plugins/Experimental/Compositing/HoldoutComposite/Source/HoldoutComposite/Private/HoldoutCompositeSceneViewExtension.cpp
@@ -12,6 +12,7 @@
 #include "FXRenderingUtils.h"
 #include "HDRHelper.h"
 #include "PixelShaderUtils.h"
+#include "PostProcess/LensDistortion.h"
 #include "PostProcess/PostProcessMaterialInputs.h"
 #include "Rendering/CustomRenderPass.h"
 #include "RenderGraphBuilder.h"
@@ -75,6 +76,8 @@ class FCompositeHoldoutCompositePS : public FGlobalShader
 		SHADER_PARAMETER_SAMPLER(SamplerState, InputSampler)
 		SHADER_PARAMETER_RDG_TEXTURE(Texture2D, CustomTexture)
 		SHADER_PARAMETER_SAMPLER(SamplerState, CustomSampler)
+		SHADER_PARAMETER_RDG_TEXTURE(Texture2D, UndistortingDisplacementTexture)
+		SHADER_PARAMETER_SAMPLER(SamplerState, UndistortingDisplacementSampler)
 		SHADER_PARAMETER_RDG_BUFFER_SRV(StructuredBuffer<float4>, EyeAdaptationBuffer)
 		SHADER_PARAMETER(FUint32Vector2, Encodings)
 		SHADER_PARAMETER(FVector2f, DisplayGamma)
@@ -377,6 +380,8 @@ FScreenPassTexture FHoldoutCompositeSceneViewExtension::PostProcessPassAfterTone
 		PermutationVector.Set<FCompositeHoldoutCompositePS::FApplyGlobalExposure>(bCompositeFollowsSceneExposure.load());
 
 		FRDGBufferRef EyeAdaptationBuffer = GraphBuilder.RegisterExternalBuffer(InView.GetEyeAdaptationBuffer(), ERDGBufferFlags::MultiFrame);
+		const FLensDistortionLUT& LensDistortionLUT = LensDistortion::GetLUTUnsafe(InView);
+		const bool bLensDistortionInTSR = (LensDistortion::GetPassLocationUnsafe(InView) == LensDistortion::EPassLocation::TSR);
 
 		FCompositeHoldoutCompositePS::FParameters* PassParameters = GraphBuilder.AllocParameters<FCompositeHoldoutCompositePS::FParameters>();
 		PassParameters->View = InView.ViewUniformBuffer;
@@ -387,6 +392,12 @@ FScreenPassTexture FHoldoutCompositeSceneViewExtension::PostProcessPassAfterTone
 		PassParameters->InputSampler = TStaticSamplerState<SF_Point>::GetRHI();
 		PassParameters->CustomTexture = CustomRenderPassTexture;
 		PassParameters->CustomSampler = TStaticSamplerState<SF_Bilinear>::GetRHI();
+		PassParameters->UndistortingDisplacementTexture = GSystemTextures.GetBlackDummy(GraphBuilder);
+		PassParameters->UndistortingDisplacementSampler = TStaticSamplerState<SF_Bilinear>::GetRHI();
+		if (LensDistortionLUT.IsEnabled() && bLensDistortionInTSR)
+		{
+			PassParameters->UndistortingDisplacementTexture = LensDistortionLUT.UndistortingDisplacementTexture;
+		}
 		PassParameters->EyeAdaptationBuffer = GraphBuilder.CreateSRV(EyeAdaptationBuffer);
 		PassParameters->Encodings = Encodings;
 		PassParameters->DisplayGamma = FVector2f(Family->RenderTarget->GetDisplayGamma(), 1.0f / Family->RenderTarget->GetDisplayGamma());
