HoldoutComposite: Temporary fix for nDisplay, until UE-209928 is addressed where custom render passes are more closely tied to view/view families, instead of the scene.

diff --git a/Engine/Plugins/Experimental/Compositing/HoldoutComposite/Source/HoldoutComposite/Private/HoldoutCompositeSceneViewExtension.cpp b/Engine/Plugins/Experimental/Compositing/HoldoutComposite/Source/HoldoutComposite/Private/HoldoutCompositeSceneViewExtension.cpp
--- a/Engine/Plugins/Experimental/Compositing/HoldoutComposite/Source/HoldoutComposite/Private/HoldoutCompositeSceneViewExtension.cpp
+++ b/Engine/Plugins/Experimental/Compositing/HoldoutComposite/Source/HoldoutComposite/Private/HoldoutCompositeSceneViewExtension.cpp
@@ -126,7 +126,13 @@ public:
 
 	virtual void OnPostRender(FRDGBuilder& GraphBuilder) override
 	{
-		ParentExtension->CollectCustomRenderTarget(ViewId, CreateDilatedTexture(GraphBuilder));
+		FRDGTextureRef DilatedCRP = CreateDilatedTexture(GraphBuilder);
+
+#if HOLDOUT_COMPOSITE_WORKAROUND_UE_209928
+		ParentExtension->CollectCustomRenderTarget(ViewId, GraphBuilder.ConvertToExternalTexture(DilatedCRP));
+#else
+		ParentExtension->CollectCustomRenderTarget(ViewId, DilatedCRP);
+#endif
 	}
 
 private:
@@ -216,11 +222,6 @@ void FHoldoutCompositeSceneViewExtension::UnregisterPrimitives(TArrayView<TSoftO
 
 }
 
-void FHoldoutCompositeSceneViewExtension::CollectCustomRenderTarget(uint32 InViewId, FRDGTextureRef InRenderTarget)
-{
-	CustomRenderTargetPerView_RenderThread.Add(InViewId, InRenderTarget);
-}
-
 bool FHoldoutCompositeSceneViewExtension::IsActiveThisFrame_Internal(const FSceneViewExtensionContext& Context) const
 {
 	bool bIsActive = FWorldSceneViewExtension::IsActiveThisFrame_Internal(Context);
@@ -319,6 +320,7 @@ void FHoldoutCompositeSceneViewExtension::SetupView(FSceneViewFamily& InViewFami
 	PassInput.CustomRenderPass = CustomRenderPass;
 	PassInput.bIsSceneCapture = true;
 
+	// TODO: Once CRPs are associated with view(family), disable & remove HOLDOUT_COMPOSITE_WORKAROUND_UE_209928.
 	WorldPtr.Get()->Scene->AddCustomRenderPass(&InViewFamily, PassInput);
 }
 
@@ -345,13 +347,20 @@ FScreenPassTexture FHoldoutCompositeSceneViewExtension::PostProcessPassAfterTone
 	FScreenPassTexture SceneColor = FScreenPassTexture::CopyFromSlice(GraphBuilder, Inputs.GetInput(EPostProcessMaterialInput::SceneColor));
 	check(SceneColor.IsValid());
 
-	const FRDGTextureRef* CustomRenderPassTexturePtr = CustomRenderTargetPerView_RenderThread.Find(InView.GetViewKey());
 	FRDGTextureRef CustomRenderPassTexture = GSystemTextures.GetBlackAlphaOneDummy(GraphBuilder);
-	
+#if HOLDOUT_COMPOSITE_WORKAROUND_UE_209928
+	const TRefCountPtr<IPooledRenderTarget>* CustomRenderPassRenderTargetPtr = CustomRenderTargetPerView_RenderThread.Find(InView.GetViewKey());
+	if (CustomRenderPassRenderTargetPtr != nullptr)
+	{
+		CustomRenderPassTexture = GraphBuilder.RegisterExternalTexture(*CustomRenderPassRenderTargetPtr);
+}
+#else
+	const FRDGTextureRef* CustomRenderPassTexturePtr = CustomRenderTargetPerView_RenderThread.Find(InView.GetViewKey());
 	if (CustomRenderPassTexturePtr && HasBeenProduced(*CustomRenderPassTexturePtr))
 	{
 		CustomRenderPassTexture = *CustomRenderPassTexturePtr;
 	}
+#endif
 
 	FGlobalShaderMap* GlobalShaderMap = GetGlobalShaderMap(InView.GetFeatureLevel());
 	const FSceneViewFamily* Family = InView.Family;
@@ -420,8 +429,6 @@ FScreenPassTexture FHoldoutCompositeSceneViewExtension::PostProcessPassAfterTone
 
 void FHoldoutCompositeSceneViewExtension::PostRenderViewFamily_RenderThread(FRDGBuilder& GraphBuilder, FSceneViewFamily& InViewFamily)
 {
-	CustomRenderTargetPerView_RenderThread.Reset();
-
 	// Cleanup invalid primitives.
 	for (auto Iter = CompositePrimitives.CreateIterator(); Iter; ++Iter)
 	{
@@ -432,3 +439,8 @@ void FHoldoutCompositeSceneViewExtension::PostRenderViewFamily_RenderThread(FRDG
 	}
 }
 
+void FHoldoutCompositeSceneViewExtension::PostRenderView_RenderThread(FRDGBuilder& GraphBuilder, FSceneView& InView)
+{
+	CustomRenderTargetPerView_RenderThread.Remove(InView.GetViewKey());
+}
+
diff --git a/Engine/Plugins/Experimental/Compositing/HoldoutComposite/Source/HoldoutComposite/Private/HoldoutCompositeSceneViewExtension.h b/Engine/Plugins/Experimental/Compositing/HoldoutComposite/Source/HoldoutComposite/Private/HoldoutCompositeSceneViewExtension.h
--- a/Engine/Plugins/Experimental/Compositing/HoldoutComposite/Source/HoldoutComposite/Private/HoldoutCompositeSceneViewExtension.h
+++ b/Engine/Plugins/Experimental/Compositing/HoldoutComposite/Source/HoldoutComposite/Private/HoldoutCompositeSceneViewExtension.h
@@ -7,6 +7,9 @@
 #include "RendererInterface.h"
 #include "Containers/ContainersFwd.h"
 
+// Temporary workaround for UE-209928 to fix holdout composite on nDisplay. Remove-disable once addressed.
+#define HOLDOUT_COMPOSITE_WORKAROUND_UE_209928 1
+
 class FViewInfo;
 
 class FHoldoutCompositeSceneViewExtension : public FWorldSceneViewExtension
@@ -22,7 +25,11 @@ public:
 	void UnregisterPrimitives(TArrayView<TSoftObjectPtr<UPrimitiveComponent>> InPrimitiveComponents, bool bInHoldoutState=false);
 
 	/* Called by the custom render pass to store its view render target for this frame. */
-	void CollectCustomRenderTarget(uint32 InViewId, FRDGTextureRef InRenderTarget);
+	template<typename T>
+	void CollectCustomRenderTarget(uint32 InViewId, T InRenderTarget)
+	{
+		CustomRenderTargetPerView_RenderThread.Add(InViewId, InRenderTarget);
+	}
 
 	//~ Begin ISceneViewExtension Interface
 	virtual int32 GetPriority() const override;
@@ -31,6 +38,7 @@ public:
 	virtual void BeginRenderViewFamily(FSceneViewFamily& InViewFamily) override {};
 	virtual void SubscribeToPostProcessingPass(EPostProcessingPass PassId, const FSceneView& InView, FAfterPassCallbackDelegateArray& InOutPassCallbacks, bool bIsPassEnabled) override;
 	virtual void PostRenderViewFamily_RenderThread(FRDGBuilder& GraphBuilder, FSceneViewFamily& InViewFamily) override;
+	virtual void PostRenderView_RenderThread(FRDGBuilder& GraphBuilder, FSceneView& InView) override;
 	//~ End ISceneViewExtension Interface
 
 protected:
@@ -44,8 +52,12 @@ private:
 	// Collection of primitives to render as a custom render pass and composite after post-processing.
 	TSet<TSoftObjectPtr<UPrimitiveComponent>> CompositePrimitives;
 
+#if HOLDOUT_COMPOSITE_WORKAROUND_UE_209928
+	TMap<uint32, TRefCountPtr<IPooledRenderTarget>> CustomRenderTargetPerView_RenderThread;
+#else
 	// Custom render pass render targets for each active view
 	TMap<uint32, FRDGTextureRef> CustomRenderTargetPerView_RenderThread;
+#endif
 
 	// Flag to enable global exposure on the composited render
 	std::atomic_bool bCompositeFollowsSceneExposure = false;
