HoldoutComposite: Primitive component search is applied recursively onto the parent component, thus fixing support for the Tex3DComponent.

We migrate the settings to a plugin specific file, and expose a list of disabled primitive classes.

diff --git a/Engine/Plugins/Experimental/Compositing/HoldoutComposite/Config/BaseHoldoutComposite.ini b/Engine/Plugins/Experimental/Compositing/HoldoutComposite/Config/BaseHoldoutComposite.ini
new file mode 100644
--- /dev/null
+++ b/Engine/Plugins/Experimental/Compositing/HoldoutComposite/Config/BaseHoldoutComposite.ini
@@ -0,0 +1,6 @@
+[/Script/HoldoutComposite.HoldoutCompositeSettings]
++DisabledPrimitiveClasses=/Script/Engine.CapsuleComponent
++DisabledPrimitiveClasses=/Script/Engine.BillboardComponent
++DisabledPrimitiveClasses=/Script/Engine.ArrowComponent
++DisabledPrimitiveClasses=/Script/Engine.DrawFrustumComponent
++DisabledPrimitiveClasses=/Script/Engine.LineBatchComponent
diff --git a/Engine/Plugins/Experimental/Compositing/HoldoutComposite/Source/HoldoutComposite/Private/HoldoutCompositeComponent.cpp b/Engine/Plugins/Experimental/Compositing/HoldoutComposite/Source/HoldoutComposite/Private/HoldoutCompositeComponent.cpp
--- a/Engine/Plugins/Experimental/Compositing/HoldoutComposite/Source/HoldoutComposite/Private/HoldoutCompositeComponent.cpp
+++ b/Engine/Plugins/Experimental/Compositing/HoldoutComposite/Source/HoldoutComposite/Private/HoldoutCompositeComponent.cpp
@@ -5,6 +5,7 @@
 #include "Components/PrimitiveComponent.h"
 #include "Engine/Engine.h"
 #include "HoldoutCompositeModule.h"
+#include "HoldoutCompositeSettings.h"
 #include "HoldoutCompositeSubsystem.h"
 
 #if WITH_EDITOR
@@ -14,6 +15,61 @@
 
 #define LOCTEXT_NAMESPACE "HoldoutComposite"
 
+namespace
+{
+	// Ensure that the primitive class is allowed by checking against plugin settings.
+	bool IsAllowedPrimitiveClass(const UPrimitiveComponent* InPrimitiveComponent)
+	{
+		const UHoldoutCompositeSettings* Settings = GetDefault<UHoldoutCompositeSettings>();
+		if (ensure(Settings))
+		{
+			for (const FSoftClassPath& ObjectClass : Settings->DisabledPrimitiveClasses)
+			{
+				UClass* TransactionClass = ObjectClass.TryLoadClass<UObject>();
+				if (TransactionClass && InPrimitiveComponent->IsA(TransactionClass))
+				{
+					return false;
+				}
+			}
+		}
+
+		return true;
+	}
+
+	// Find list of primitives from the parent component and its children. We need to traverse children to support objects such as UText3DComponent.
+	static TArray<TSoftObjectPtr<UPrimitiveComponent>> FindPrimitiveComponents(USceneComponent* InParentComponent)
+	{
+		TArray<TSoftObjectPtr<UPrimitiveComponent>> OutPrimitiveComponents;
+
+		if (IsValid(InParentComponent))
+		{
+			if (UPrimitiveComponent* PrimitiveComponent = Cast<UPrimitiveComponent>(InParentComponent))
+			{
+				if (IsAllowedPrimitiveClass(PrimitiveComponent))
+				{
+					OutPrimitiveComponents.Add(PrimitiveComponent);
+				}
+			}
+
+			TArray<USceneComponent*> ParentChildComponents;
+			InParentComponent->GetChildrenComponents(true /*bIncludeAllDescendants*/, ParentChildComponents);
+
+			for (USceneComponent* ParentChild : ParentChildComponents)
+			{
+				if (UPrimitiveComponent* PrimitiveComponent = Cast<UPrimitiveComponent>(ParentChild))
+				{
+					if (IsAllowedPrimitiveClass(PrimitiveComponent))
+					{
+						OutPrimitiveComponents.Add(PrimitiveComponent);
+					}
+				}
+			}
+		}
+
+		return OutPrimitiveComponents;
+	}
+}
+
 UHoldoutCompositeComponent::UHoldoutCompositeComponent(const FObjectInitializer& ObjectInitializer)
 	: Super(ObjectInitializer)
 {
@@ -47,46 +103,48 @@ void UHoldoutCompositeComponent::OnAttachmentChanged()
 
 	UnregisterCompositeImpl();
 
-	const USceneComponent* SceneComponent = GetAttachParent();
+	USceneComponent* SceneComponent = GetAttachParent();
 	if (IsValid(SceneComponent))
 	{
-		const UPrimitiveComponent* ParentPrimitiveComponent = Cast<UPrimitiveComponent>(SceneComponent);
-		if (IsValid(ParentPrimitiveComponent))
-		{
-			RegisterCompositeImpl();
-		}
-		else
+		TArray<TSoftObjectPtr<UPrimitiveComponent>> ParentPrimitives = FindPrimitiveComponents(SceneComponent);
+		if (ParentPrimitives.IsEmpty())
 		{
 #if WITH_EDITOR
 			FNotificationInfo Info(LOCTEXT("CompositeParentNotification",
-				"The composite component must be parented to a primitive component."));
+				"The composite component must be parented to a primitive component (or one that has primitives)."));
 			Info.ExpireDuration = 5.0f;
 
 			FSlateNotificationManager::Get().AddNotification(Info);
 #endif
 		}
+		else
+		{
+			RegisterCompositeImpl();
+		}
 	}
+
+
 }
 
 void UHoldoutCompositeComponent::RegisterCompositeImpl()
 {
 	UHoldoutCompositeSubsystem* Subsystem = UWorld::GetSubsystem<UHoldoutCompositeSubsystem>(GetWorld());
-	UPrimitiveComponent* ParentPrimitiveComponent = Cast<UPrimitiveComponent>(GetAttachParent());
+	TArray<TSoftObjectPtr<UPrimitiveComponent>> ParentPrimitives = FindPrimitiveComponents(GetAttachParent());
 
-	if (IsValid(Subsystem) && IsValid(ParentPrimitiveComponent))
+	if (IsValid(Subsystem) && !ParentPrimitives.IsEmpty())
 	{
-		Subsystem->RegisterPrimitive(ParentPrimitiveComponent);
+		Subsystem->RegisterPrimitives(ParentPrimitives);
 	}
 }
 
 void UHoldoutCompositeComponent::UnregisterCompositeImpl()
 {
 	UHoldoutCompositeSubsystem* Subsystem = UWorld::GetSubsystem<UHoldoutCompositeSubsystem>(GetWorld());
-	UPrimitiveComponent* ParentPrimitiveComponent = Cast<UPrimitiveComponent>(GetAttachParent());
+	TArray<TSoftObjectPtr<UPrimitiveComponent>> ParentPrimitives = FindPrimitiveComponents(GetAttachParent());
 
-	if (IsValid(Subsystem) && IsValid(ParentPrimitiveComponent))
+	if (IsValid(Subsystem) && !ParentPrimitives.IsEmpty())
 	{
-		Subsystem->UnregisterPrimitive(ParentPrimitiveComponent);
+		Subsystem->UnregisterPrimitives(ParentPrimitives);
 	}
 }
 
diff --git a/Engine/Plugins/Experimental/Compositing/HoldoutComposite/Source/HoldoutComposite/Public/HoldoutCompositeSettings.h b/Engine/Plugins/Experimental/Compositing/HoldoutComposite/Source/HoldoutComposite/Public/HoldoutCompositeSettings.h
--- a/Engine/Plugins/Experimental/Compositing/HoldoutComposite/Source/HoldoutComposite/Public/HoldoutCompositeSettings.h
+++ b/Engine/Plugins/Experimental/Compositing/HoldoutComposite/Source/HoldoutComposite/Public/HoldoutCompositeSettings.h
@@ -12,8 +12,8 @@
 /**
  * Settings for the HoldoutComposite module.
  */
-UCLASS(config = Engine)
-class HOLDOUTCOMPOSITE_API UHoldoutCompositeSettings : public UDeveloperSettings
+UCLASS(config = HoldoutComposite, defaultconfig, meta = (DisplayName = "Holdout Composite"), MinimalAPI)
+class UHoldoutCompositeSettings : public UDeveloperSettings
 {
 	GENERATED_BODY()
 
@@ -36,5 +36,9 @@ public:
 	/** Composite (scene view extension) pass priority, which defaults to before OpenColorIO. */
 	UPROPERTY(config, EditAnywhere, Category = General, AdvancedDisplay)
 	int32 SceneViewExtensionPriority;
+
+	/** Primitive component classes that do not support the holdout composite.*/
+	UPROPERTY(config, EditAnywhere, Category = General)
+	TArray<FSoftClassPath> DisabledPrimitiveClasses;
 };
 
