Alpha: Minor change to propagate alpha checks: favor enum use over hard-coded integer.

Should be helpful for the coming transition of the alpha project setting back to a boolean.

diff --git a/Engine/Plugins/Experimental/Avalanche/Source/AvalancheMediaEditor/Private/Rundown/Preview/SAvaRundownPagePreview.cpp b/Engine/Plugins/Experimental/Avalanche/Source/AvalancheMediaEditor/Private/Rundown/Preview/SAvaRundownPagePreview.cpp
--- a/Engine/Plugins/Experimental/Avalanche/Source/AvalancheMediaEditor/Private/Rundown/Preview/SAvaRundownPagePreview.cpp
+++ b/Engine/Plugins/Experimental/Avalanche/Source/AvalancheMediaEditor/Private/Rundown/Preview/SAvaRundownPagePreview.cpp
@@ -6,6 +6,7 @@
 #include "AvaMediaSettings.h"
 #include "Broadcast/OutputDevices/Slate/SAvaBroadcastCaptureImage.h"
 #include "Brushes/SlateImageBrush.h"
+#include "Engine/RendererSettings.h"
 #include "Engine/TextureRenderTarget2D.h"
 #include "Framework/Commands/UICommandList.h"
 #include "Framework/MultiBox/MultiBoxBuilder.h"
@@ -303,7 +304,7 @@ void SAvaRundownPagePreview::HandleCheckerboardActionExecute() const
 	if (bShowCheckerBoard)
 	{
 		const IConsoleVariable* PropagateAlphaCVar = IConsoleManager::Get().FindConsoleVariable(TEXT("r.PostProcessing.PropagateAlpha"));
-		if (PropagateAlphaCVar && PropagateAlphaCVar->GetInt() != 2)
+		if (PropagateAlphaCVar && PropagateAlphaCVar->GetInt() != static_cast<int32>(EAlphaChannelMode::AllowThroughTonemapper))
 		{
 			const FText NotificationText = LOCTEXT("AlphaSupport",
 				"An output requested Alpha Support but the required project setting is not enabled!\n"
diff --git a/Engine/Plugins/Experimental/Compositing/HoldoutComposite/Source/HoldoutComposite/Private/HoldoutCompositeSubsystem.cpp b/Engine/Plugins/Experimental/Compositing/HoldoutComposite/Source/HoldoutComposite/Private/HoldoutCompositeSubsystem.cpp
--- a/Engine/Plugins/Experimental/Compositing/HoldoutComposite/Source/HoldoutComposite/Private/HoldoutCompositeSubsystem.cpp
+++ b/Engine/Plugins/Experimental/Compositing/HoldoutComposite/Source/HoldoutComposite/Private/HoldoutCompositeSubsystem.cpp
@@ -4,6 +4,8 @@
 #include "HoldoutCompositeModule.h"
 #include "HoldoutCompositeSceneViewExtension.h"
 
+#include "Engine/RendererSettings.h"
+
 UHoldoutCompositeSubsystem::UHoldoutCompositeSubsystem()
 {
 }
@@ -36,11 +38,11 @@ void UHoldoutCompositeSubsystem::RegisterPrimitives(TArrayView<TSoftObjectPtr<UP
 {
 	// The compositing relies on alpha preserved through the "tonemapper" post-processing step.
 	static IConsoleVariable* CVarPropagateAlpha = IConsoleManager::Get().FindConsoleVariable(TEXT("r.PostProcessing.PropagateAlpha"));
-	if (CVarPropagateAlpha && CVarPropagateAlpha->GetInt() != 2 /*EAlphaChannelMode::AllowThroughTonemapper*/)
+	if (CVarPropagateAlpha && CVarPropagateAlpha->GetInt() != static_cast<int32>(EAlphaChannelMode::AllowThroughTonemapper))
 	{
 		UE_CALL_ONCE([]()
 			{
-				UE_LOG(LogHoldoutComposite, Warning, TEXT("Holdout composite is disabled until r.PostProcessing.PropagateAlpha=2 is set."));
+				UE_LOG(LogHoldoutComposite, Warning, TEXT("Holdout composite is disabled until alpha is enabled through post-processing."));
 			}
 		);
 		return;
