HoldoutComposite: Add built-in support for SSR (onto CG).

Exposed as a project setting, and enabled by default. Cleanup/group of existing shaders in passing.

We also hide properties on the holdout composite component, and revert auto-exposure on the composite to disabled by default.

diff --git a/Engine/Plugins/Experimental/Compositing/HoldoutComposite/Shaders/Private/HoldoutCompositeFinal.usf b/Engine/Plugins/Experimental/Compositing/HoldoutComposite/Shaders/Private/HoldoutCompositeFinal.usf
--- a/Engine/Plugins/Experimental/Compositing/HoldoutComposite/Shaders/Private/HoldoutCompositeFinal.usf
+++ b/Engine/Plugins/Experimental/Compositing/HoldoutComposite/Shaders/Private/HoldoutCompositeFinal.usf
@@ -15,6 +15,9 @@ SamplerState InputSampler;
 Texture2D CustomTexture;
 SamplerState CustomSampler;
 
+Texture2D<float2> DistortingDisplacementTexture;
+SamplerState DistortingDisplacementSampler;
+
 Texture2D<float2> UndistortingDisplacementTexture;
 SamplerState UndistortingDisplacementSampler;
 
@@ -36,7 +39,7 @@ void MainPS(
     const float2 CustomUV = (UndistortedViewportUV * Custom_ViewportSize + Custom_ViewportMin) * Custom_ExtentInverse;
     float4 CustomColor = Texture2DSample(CustomTexture, CustomSampler, CustomUV);
 
-#if PERMUTATION_GLOBAL_EXPOSURE
+#if USE_GLOBAL_EXPOSURE
     const float GlobalExposure = EyeAdaptationBuffer[0].x;
     CustomColor.rgb *= GlobalExposure;
 #endif
diff --git a/Engine/Plugins/Experimental/Compositing/HoldoutComposite/Shaders/Private/HoldoutCompositeSSRInput.usf b/Engine/Plugins/Experimental/Compositing/HoldoutComposite/Shaders/Private/HoldoutCompositeSSRInput.usf
new file mode 100644
--- /dev/null
+++ b/Engine/Plugins/Experimental/Compositing/HoldoutComposite/Shaders/Private/HoldoutCompositeSSRInput.usf
@@ -0,0 +1,42 @@
+// Copyright Epic Games, Inc. All Rights Reserved.
+
+#include "/Engine/Private/Common.ush"
+#include "/Engine/Private/ScreenPass.ush"
+#include "/Engine/Private/LensDistortion.ush"
+
+SCREEN_PASS_TEXTURE_VIEWPORT(Input)
+SCREEN_PASS_TEXTURE_VIEWPORT(Custom)
+SCREEN_PASS_TEXTURE_VIEWPORT(Output)
+
+Texture2D InputTexture;
+SamplerState InputSampler;
+
+Texture2D CustomTexture;
+SamplerState CustomSampler;
+
+Texture2D<float2> DistortingDisplacementTexture;
+SamplerState DistortingDisplacementSampler;
+
+Texture2D<float2> UndistortingDisplacementTexture;
+SamplerState UndistortingDisplacementSampler;
+
+float LastGlobalExposure;
+
+void MainPS(
+    float4 SvPosition : SV_POSITION,
+    out float4 OutColor : SV_Target0
+    )
+{
+    float2 ViewportUV = (SvPosition.xy - Output_ViewportMin) * Output_ViewportSizeInverse;
+    const float2 DistortedViewportUV = ApplyLensDistortionOnViewportUV(DistortingDisplacementTexture, DistortingDisplacementSampler, ViewportUV);
+    
+    const float2 InputUV = (DistortedViewportUV * Input_ViewportSize + Input_ViewportMin) * Input_ExtentInverse;
+    float4 SceneColor = Texture2DSample(InputTexture, InputSampler, InputUV);
+
+    const float2 CustomUV = (ViewportUV * Custom_ViewportSize + Custom_ViewportMin) * Custom_ExtentInverse;
+    float4 CustomColor = Texture2DSample(CustomTexture, CustomSampler, CustomUV);
+
+    CustomColor.rgb *= LastGlobalExposure;
+    OutColor.rgb = CustomColor.rgb * SceneColor.a + SceneColor.rgb;
+    OutColor.a = CustomColor.a * SceneColor.a;
+}
diff --git a/Engine/Plugins/Experimental/Compositing/HoldoutComposite/Source/HoldoutComposite/Private/HoldoutCompositeSceneViewExtension.cpp b/Engine/Plugins/Experimental/Compositing/HoldoutComposite/Source/HoldoutComposite/Private/HoldoutCompositeSceneViewExtension.cpp
--- a/Engine/Plugins/Experimental/Compositing/HoldoutComposite/Source/HoldoutComposite/Private/HoldoutCompositeSceneViewExtension.cpp
+++ b/Engine/Plugins/Experimental/Compositing/HoldoutComposite/Source/HoldoutComposite/Private/HoldoutCompositeSceneViewExtension.cpp
@@ -25,6 +25,7 @@
 #include "SceneRendering.h"
 
 DECLARE_GPU_STAT_NAMED(FHoldoutCompositeDilate, TEXT("HoldoutComposite.Dilate"));
+DECLARE_GPU_STAT_NAMED(FHoldoutCompositeSSRInput, TEXT("HoldoutComposite.SSRInput"));
 DECLARE_GPU_STAT_NAMED(FHoldoutCompositeFinal, TEXT("HoldoutComposite.Final"));
 
 namespace HoldoutComposite
@@ -37,10 +38,10 @@ namespace HoldoutComposite
 	};
 }
 
-class FDilateHoldoutCompositeCS : public FGlobalShader
+class FHoldoutCompositeDilateShader : public FGlobalShader
 {
-	DECLARE_GLOBAL_SHADER(FDilateHoldoutCompositeCS);
-	SHADER_USE_PARAMETER_STRUCT(FDilateHoldoutCompositeCS, FGlobalShader);
+	DECLARE_GLOBAL_SHADER(FHoldoutCompositeDilateShader);
+	SHADER_USE_PARAMETER_STRUCT(FHoldoutCompositeDilateShader, FGlobalShader);
 
 	BEGIN_SHADER_PARAMETER_STRUCT(FParameters, )
 		SHADER_PARAMETER_RDG_TEXTURE(Texture2D, InputTexture)
@@ -57,35 +58,53 @@ class FDilateHoldoutCompositeCS : public FGlobalShader
 
 	static const uint32 ThreadGroupSize = 16;
 };
-IMPLEMENT_GLOBAL_SHADER(FDilateHoldoutCompositeCS, "/Plugin/HoldoutComposite/Private/HoldoutCompositeDilate.usf", "MainCS", SF_Compute);
+IMPLEMENT_GLOBAL_SHADER(FHoldoutCompositeDilateShader, "/Plugin/HoldoutComposite/Private/HoldoutCompositeDilate.usf", "MainCS", SF_Compute);
+
+BEGIN_SHADER_PARAMETER_STRUCT(FHoldoutCompositeCommonParameters, )
+	SHADER_PARAMETER_STRUCT(FScreenPassTextureViewportParameters, Input)
+	SHADER_PARAMETER_STRUCT(FScreenPassTextureViewportParameters, Custom)
+	SHADER_PARAMETER_STRUCT(FScreenPassTextureViewportParameters, Output)
+	SHADER_PARAMETER_RDG_TEXTURE(Texture2D, InputTexture)
+	SHADER_PARAMETER_SAMPLER(SamplerState, InputSampler)
+	SHADER_PARAMETER_RDG_TEXTURE(Texture2D, CustomTexture)
+	SHADER_PARAMETER_SAMPLER(SamplerState, CustomSampler)
+	SHADER_PARAMETER_RDG_TEXTURE(Texture2D, DistortingDisplacementTexture)
+	SHADER_PARAMETER_SAMPLER(SamplerState, DistortingDisplacementSampler)
+	SHADER_PARAMETER_RDG_TEXTURE(Texture2D, UndistortingDisplacementTexture)
+	SHADER_PARAMETER_SAMPLER(SamplerState, UndistortingDisplacementSampler)
+END_SHADER_PARAMETER_STRUCT()
+
+class FHoldoutCompositeSSRInputShader : public FGlobalShader
+{
+	DECLARE_GLOBAL_SHADER(FHoldoutCompositeSSRInputShader);
+	SHADER_USE_PARAMETER_STRUCT(FHoldoutCompositeSSRInputShader, FGlobalShader);
 
-class FCompositeHoldoutCompositePS : public FGlobalShader
+	BEGIN_SHADER_PARAMETER_STRUCT(FParameters, )
+		SHADER_PARAMETER_STRUCT_INCLUDE(FHoldoutCompositeCommonParameters, Common)
+		SHADER_PARAMETER(float, LastGlobalExposure)
+		RENDER_TARGET_BINDING_SLOTS()
+	END_SHADER_PARAMETER_STRUCT()
+};
+IMPLEMENT_GLOBAL_SHADER(FHoldoutCompositeSSRInputShader, "/Plugin/HoldoutComposite/Private/HoldoutCompositeSSRInput.usf", "MainPS", SF_Pixel);
+
+class FHoldoutCompositeFinalShader : public FGlobalShader
 {
-	DECLARE_GLOBAL_SHADER(FCompositeHoldoutCompositePS);
-	SHADER_USE_PARAMETER_STRUCT(FCompositeHoldoutCompositePS, FGlobalShader);
+	DECLARE_GLOBAL_SHADER(FHoldoutCompositeFinalShader);
+	SHADER_USE_PARAMETER_STRUCT(FHoldoutCompositeFinalShader, FGlobalShader);
 
-	class FApplyGlobalExposure : SHADER_PERMUTATION_BOOL("PERMUTATION_GLOBAL_EXPOSURE");
-	using FPermutationDomain = TShaderPermutationDomain<FApplyGlobalExposure>;
+	class FUseGlobalExposure : SHADER_PERMUTATION_BOOL("USE_GLOBAL_EXPOSURE");
+	using FPermutationDomain = TShaderPermutationDomain<FUseGlobalExposure>;
 
 	BEGIN_SHADER_PARAMETER_STRUCT(FParameters, )
 		SHADER_PARAMETER_STRUCT_REF(FViewUniformShaderParameters, View)
-		SHADER_PARAMETER_STRUCT(FScreenPassTextureViewportParameters, Input)
-		SHADER_PARAMETER_STRUCT(FScreenPassTextureViewportParameters, Custom)
-		SHADER_PARAMETER_STRUCT(FScreenPassTextureViewportParameters, Output)
-		SHADER_PARAMETER_RDG_TEXTURE(Texture2D, InputTexture)
-		SHADER_PARAMETER_SAMPLER(SamplerState, InputSampler)
-		SHADER_PARAMETER_RDG_TEXTURE(Texture2D, CustomTexture)
-		SHADER_PARAMETER_SAMPLER(SamplerState, CustomSampler)
-		SHADER_PARAMETER_RDG_TEXTURE(Texture2D, UndistortingDisplacementTexture)
-		SHADER_PARAMETER_SAMPLER(SamplerState, UndistortingDisplacementSampler)
+		SHADER_PARAMETER_STRUCT_INCLUDE(FHoldoutCompositeCommonParameters, Common)
 		SHADER_PARAMETER_RDG_BUFFER_SRV(StructuredBuffer<float4>, EyeAdaptationBuffer)
 		SHADER_PARAMETER(FUint32Vector2, Encodings)
 		SHADER_PARAMETER(FVector2f, DisplayGamma)
-
 		RENDER_TARGET_BINDING_SLOTS()
 	END_SHADER_PARAMETER_STRUCT()
 };
-IMPLEMENT_GLOBAL_SHADER(FCompositeHoldoutCompositePS, "/Plugin/HoldoutComposite/Private/HoldoutCompositeFinal.usf", "MainPS", SF_Pixel);
+IMPLEMENT_GLOBAL_SHADER(FHoldoutCompositeFinalShader, "/Plugin/HoldoutComposite/Private/HoldoutCompositeFinal.usf", "MainPS", SF_Pixel);
 
 class FHoldoutCompositeCustomRenderPass : public FCustomRenderPassBase
 {
@@ -150,19 +169,19 @@ private:
 
 		// Async compute dilation pass
 		{
-			FDilateHoldoutCompositeCS::FParameters* PassParameters = GraphBuilder.AllocParameters<FDilateHoldoutCompositeCS::FParameters>();
+			FHoldoutCompositeDilateShader::FParameters* PassParameters = GraphBuilder.AllocParameters<FHoldoutCompositeDilateShader::FParameters>();
 			PassParameters->InputTexture = RenderTargetTexture;
 			PassParameters->RWOutputTexture = GraphBuilder.CreateUAV(DilatedTexture);
 			PassParameters->Dimensions = TextureSize;
 
-			TShaderMapRef<FDilateHoldoutCompositeCS> ComputeShader(GlobalShaderMap);
+			TShaderMapRef<FHoldoutCompositeDilateShader> ComputeShader(GlobalShaderMap);
 			FComputeShaderUtils::AddPass(
 				GraphBuilder,
 				RDG_EVENT_NAME("HoldoutComposite.Dilate (%dx%d)", TextureSize.X, TextureSize.Y),
 				GSupportsEfficientAsyncCompute ? ERDGPassFlags::AsyncCompute : ERDGPassFlags::Compute,
 				ComputeShader,
 				PassParameters,
-				FComputeShaderUtils::GetGroupCount(TextureSize, FDilateHoldoutCompositeCS::ThreadGroupSize)
+				FComputeShaderUtils::GetGroupCount(TextureSize, FHoldoutCompositeDilateShader::ThreadGroupSize)
 			);
 		}
 
@@ -251,6 +270,7 @@ void FHoldoutCompositeSceneViewExtension::SetupViewFamily(FSceneViewFamily& InVi
 	if (Settings != nullptr)
 	{
 		bCompositeFollowsSceneExposure = Settings->bCompositeFollowsSceneExposure;
+		bCompositeSupportsSSR = Settings->bCompositeSupportsSSR;
 	}
 }
 
@@ -332,28 +352,26 @@ void FHoldoutCompositeSceneViewExtension::SubscribeToPostProcessingPass(EPostPro
 		return;
 	}
 
-	if (PassId == EPostProcessingPass::Tonemap)
+	if (PassId == EPostProcessingPass::SSRInput && bCompositeSupportsSSR.load())
+	{
+		InOutPassCallbacks.Add(FAfterPassCallbackDelegate::CreateRaw(this, &FHoldoutCompositeSceneViewExtension::PostProcessPassSSRInput_RenderThread));
+	}
+	else if (PassId == EPostProcessingPass::Tonemap)
 	{
 		InOutPassCallbacks.Add(FAfterPassCallbackDelegate::CreateRaw(this, &FHoldoutCompositeSceneViewExtension::PostProcessPassAfterTonemap_RenderThread));
 	}
 }
 
-FScreenPassTexture FHoldoutCompositeSceneViewExtension::PostProcessPassAfterTonemap_RenderThread(FRDGBuilder& GraphBuilder, const FSceneView& InView, const FPostProcessMaterialInputs& Inputs)
+FRDGTextureRef FHoldoutCompositeSceneViewExtension::GetCustomRenderPassTexture(FRDGBuilder& GraphBuilder, const FSceneView& InView) const
 {
-	RDG_EVENT_SCOPE(GraphBuilder, "HoldoutComposite.Final");
-	RDG_GPU_STAT_SCOPE(GraphBuilder, FHoldoutCompositeFinal);
-	using namespace HoldoutComposite;
-
-	FScreenPassTexture SceneColor = FScreenPassTexture::CopyFromSlice(GraphBuilder, Inputs.GetInput(EPostProcessMaterialInput::SceneColor));
-	check(SceneColor.IsValid());
-
 	FRDGTextureRef CustomRenderPassTexture = GSystemTextures.GetBlackAlphaOneDummy(GraphBuilder);
+
 #if HOLDOUT_COMPOSITE_WORKAROUND_UE_209928
 	const TRefCountPtr<IPooledRenderTarget>* CustomRenderPassRenderTargetPtr = CustomRenderTargetPerView_RenderThread.Find(InView.GetViewKey());
 	if (CustomRenderPassRenderTargetPtr != nullptr)
 	{
 		CustomRenderPassTexture = GraphBuilder.RegisterExternalTexture(*CustomRenderPassRenderTargetPtr);
-}
+	}
 #else
 	const FRDGTextureRef* CustomRenderPassTexturePtr = CustomRenderTargetPerView_RenderThread.Find(InView.GetViewKey());
 	if (CustomRenderPassTexturePtr && HasBeenProduced(*CustomRenderPassTexturePtr))
@@ -361,6 +379,79 @@ FScreenPassTexture FHoldoutCompositeSceneViewExtension::PostProcessPassAfterTone
 		CustomRenderPassTexture = *CustomRenderPassTexturePtr;
 	}
 #endif
+	
+	return CustomRenderPassTexture;
+}
+
+FHoldoutCompositeCommonParameters FHoldoutCompositeSceneViewExtension::BuildCommonCompositeParameters(FRDGBuilder& GraphBuilder, const FSceneView& InView, const FScreenPassTexture& SceneColor, const FScreenPassRenderTarget& Output, bool bIsSceneColorUndistorted)
+{
+	const FLensDistortionLUT& LensDistortionLUT = LensDistortion::GetLUTUnsafe(InView);
+	const bool bLensDistortionInTSR = (LensDistortion::GetPassLocationUnsafe(InView) == LensDistortion::EPassLocation::TSR);
+
+	FRDGTextureRef CustomRenderPassTexture = GetCustomRenderPassTexture(GraphBuilder, InView);
+
+	FHoldoutCompositeCommonParameters CommonParameters;
+	CommonParameters.Input = GetScreenPassTextureViewportParameters(FScreenPassTextureViewport(SceneColor));
+	CommonParameters.Custom = GetScreenPassTextureViewportParameters(FScreenPassTextureViewport(CustomRenderPassTexture));
+	CommonParameters.Output = GetScreenPassTextureViewportParameters(FScreenPassTextureViewport(Output));
+	CommonParameters.InputTexture = SceneColor.Texture;
+	CommonParameters.InputSampler = bIsSceneColorUndistorted ? TStaticSamplerState<SF_Bilinear, AM_Mirror, AM_Mirror>::GetRHI() : TStaticSamplerState<SF_Point>::GetRHI();
+	CommonParameters.CustomTexture = CustomRenderPassTexture;
+	CommonParameters.CustomSampler = TStaticSamplerState<SF_Bilinear>::GetRHI();
+	CommonParameters.UndistortingDisplacementTexture = GSystemTextures.GetBlackDummy(GraphBuilder);
+	CommonParameters.UndistortingDisplacementSampler = TStaticSamplerState<SF_Bilinear>::GetRHI();
+	CommonParameters.DistortingDisplacementTexture = GSystemTextures.GetBlackDummy(GraphBuilder);
+	CommonParameters.DistortingDisplacementSampler = TStaticSamplerState<SF_Bilinear>::GetRHI();
+	if (LensDistortionLUT.IsEnabled() && bLensDistortionInTSR)
+	{
+		CommonParameters.DistortingDisplacementTexture = LensDistortionLUT.DistortingDisplacementTexture;
+		CommonParameters.UndistortingDisplacementTexture = LensDistortionLUT.UndistortingDisplacementTexture;
+	}
+
+	return CommonParameters;
+}
+
+FScreenPassTexture FHoldoutCompositeSceneViewExtension::PostProcessPassSSRInput_RenderThread(FRDGBuilder& GraphBuilder, const FSceneView& InView, const FPostProcessMaterialInputs& Inputs)
+{
+	RDG_EVENT_SCOPE(GraphBuilder, "HoldoutComposite.SSRInput");
+	RDG_GPU_STAT_SCOPE(GraphBuilder, FHoldoutCompositeSSRInput);
+
+	FGlobalShaderMap* GlobalShaderMap = GetGlobalShaderMap(InView.GetFeatureLevel());
+
+	FScreenPassTexture SceneColor = FScreenPassTexture::CopyFromSlice(GraphBuilder, Inputs.GetInput(EPostProcessMaterialInput::SceneColor));
+	check(SceneColor.IsValid());
+	FScreenPassRenderTarget Output = FScreenPassRenderTarget::CreateFromInput(GraphBuilder, SceneColor, InView.GetOverwriteLoadAction(), TEXT("HoldoutCompositeSSRInputRT"));
+
+	FHoldoutCompositeSSRInputShader::FParameters* PassParameters = GraphBuilder.AllocParameters<FHoldoutCompositeSSRInputShader::FParameters>();
+	constexpr bool bIsSceneColorUndistorted = true;
+	PassParameters->Common = BuildCommonCompositeParameters(GraphBuilder, InView, SceneColor, Output, bIsSceneColorUndistorted);
+	PassParameters->LastGlobalExposure = bCompositeFollowsSceneExposure.load() ? InView.GetLastEyeAdaptationExposure() : 1.0f;
+	PassParameters->RenderTargets[0] = Output.GetRenderTargetBinding();
+
+	TShaderMapRef<FHoldoutCompositeSSRInputShader> PixelShader(GlobalShaderMap);
+	FPixelShaderUtils::AddFullscreenPass(
+		GraphBuilder,
+		GlobalShaderMap,
+		RDG_EVENT_NAME("HoldoutComposite.SSRInput (%dx%d) PS",
+			Output.ViewRect.Width(), Output.ViewRect.Height()),
+		PixelShader,
+		PassParameters,
+		Output.ViewRect
+	);
+
+	return Output;
+}
+
+FScreenPassTexture FHoldoutCompositeSceneViewExtension::PostProcessPassAfterTonemap_RenderThread(FRDGBuilder& GraphBuilder, const FSceneView& InView, const FPostProcessMaterialInputs& Inputs)
+{
+	RDG_EVENT_SCOPE(GraphBuilder, "HoldoutComposite.Final");
+	RDG_GPU_STAT_SCOPE(GraphBuilder, FHoldoutCompositeFinal);
+	using namespace HoldoutComposite;
+
+	FScreenPassTexture SceneColor = FScreenPassTexture::CopyFromSlice(GraphBuilder, Inputs.GetInput(EPostProcessMaterialInput::SceneColor));
+	check(SceneColor.IsValid());
+
+	FRDGTextureRef CustomRenderPassTexture = GetCustomRenderPassTexture(GraphBuilder, InView);
 
 	FGlobalShaderMap* GlobalShaderMap = GetGlobalShaderMap(InView.GetFeatureLevel());
 	const FSceneViewFamily* Family = InView.Family;
@@ -385,34 +476,22 @@ FScreenPassTexture FHoldoutCompositeSceneViewExtension::PostProcessPassAfterTone
 
 	// Compositing pass
 	{
-		FCompositeHoldoutCompositePS::FPermutationDomain PermutationVector;
-		PermutationVector.Set<FCompositeHoldoutCompositePS::FApplyGlobalExposure>(bCompositeFollowsSceneExposure.load());
-
-		FRDGBufferRef EyeAdaptationBuffer = GraphBuilder.RegisterExternalBuffer(InView.GetEyeAdaptationBuffer(), ERDGBufferFlags::MultiFrame);
-		const FLensDistortionLUT& LensDistortionLUT = LensDistortion::GetLUTUnsafe(InView);
-		const bool bLensDistortionInTSR = (LensDistortion::GetPassLocationUnsafe(InView) == LensDistortion::EPassLocation::TSR);
+		FHoldoutCompositeFinalShader::FPermutationDomain PermutationVector;
+		PermutationVector.Set<FHoldoutCompositeFinalShader::FUseGlobalExposure>(bCompositeFollowsSceneExposure.load());
 
-		FCompositeHoldoutCompositePS::FParameters* PassParameters = GraphBuilder.AllocParameters<FCompositeHoldoutCompositePS::FParameters>();
+		FHoldoutCompositeFinalShader::FParameters* PassParameters = GraphBuilder.AllocParameters<FHoldoutCompositeFinalShader::FParameters>();
 		PassParameters->View = InView.ViewUniformBuffer;
-		PassParameters->Input = GetScreenPassTextureViewportParameters(FScreenPassTextureViewport(SceneColor));
-		PassParameters->Custom = GetScreenPassTextureViewportParameters(FScreenPassTextureViewport(CustomRenderPassTexture));
-		PassParameters->Output = GetScreenPassTextureViewportParameters(FScreenPassTextureViewport(Output));
-		PassParameters->InputTexture = SceneColor.Texture;
-		PassParameters->InputSampler = TStaticSamplerState<SF_Point>::GetRHI();
-		PassParameters->CustomTexture = CustomRenderPassTexture;
-		PassParameters->CustomSampler = TStaticSamplerState<SF_Bilinear>::GetRHI();
-		PassParameters->UndistortingDisplacementTexture = GSystemTextures.GetBlackDummy(GraphBuilder);
-		PassParameters->UndistortingDisplacementSampler = TStaticSamplerState<SF_Bilinear>::GetRHI();
-		if (LensDistortionLUT.IsEnabled() && bLensDistortionInTSR)
+		PassParameters->Common = BuildCommonCompositeParameters(GraphBuilder, InView, SceneColor, Output);
+		if (bCompositeFollowsSceneExposure.load())
 		{
-			PassParameters->UndistortingDisplacementTexture = LensDistortionLUT.UndistortingDisplacementTexture;
+			FRDGBufferRef EyeAdaptationBuffer = GraphBuilder.RegisterExternalBuffer(InView.GetEyeAdaptationBuffer(), ERDGBufferFlags::MultiFrame);
+			PassParameters->EyeAdaptationBuffer = GraphBuilder.CreateSRV(EyeAdaptationBuffer);
 		}
-		PassParameters->EyeAdaptationBuffer = GraphBuilder.CreateSRV(EyeAdaptationBuffer);
 		PassParameters->Encodings = Encodings;
 		PassParameters->DisplayGamma = FVector2f(Family->RenderTarget->GetDisplayGamma(), 1.0f / Family->RenderTarget->GetDisplayGamma());
 		PassParameters->RenderTargets[0] = Output.GetRenderTargetBinding();
 
-		TShaderMapRef<FCompositeHoldoutCompositePS> PixelShader(GlobalShaderMap, PermutationVector);
+		TShaderMapRef<FHoldoutCompositeFinalShader> PixelShader(GlobalShaderMap, PermutationVector);
 		FPixelShaderUtils::AddFullscreenPass(
 			GraphBuilder,
 			GlobalShaderMap,
diff --git a/Engine/Plugins/Experimental/Compositing/HoldoutComposite/Source/HoldoutComposite/Private/HoldoutCompositeSceneViewExtension.h b/Engine/Plugins/Experimental/Compositing/HoldoutComposite/Source/HoldoutComposite/Private/HoldoutCompositeSceneViewExtension.h
--- a/Engine/Plugins/Experimental/Compositing/HoldoutComposite/Source/HoldoutComposite/Private/HoldoutCompositeSceneViewExtension.h
+++ b/Engine/Plugins/Experimental/Compositing/HoldoutComposite/Source/HoldoutComposite/Private/HoldoutCompositeSceneViewExtension.h
@@ -11,6 +11,8 @@
 #define HOLDOUT_COMPOSITE_WORKAROUND_UE_209928 1
 
 class FViewInfo;
+struct FScreenPassRenderTarget;
+struct FScreenPassTexture;
 
 class FHoldoutCompositeSceneViewExtension : public FWorldSceneViewExtension
 {
@@ -47,7 +49,12 @@ protected:
 	//~ End ISceneViewExtension Interface
 
 private:
-	FScreenPassTexture PostProcessPassAfterTonemap_RenderThread(FRDGBuilder& GraphBuilder, const FSceneView& View, const FPostProcessMaterialInputs& Inputs);
+
+	FRDGTextureRef GetCustomRenderPassTexture(FRDGBuilder& GraphBuilder, const FSceneView& InView) const;
+	class FHoldoutCompositeCommonParameters BuildCommonCompositeParameters(FRDGBuilder& GraphBuilder, const FSceneView& InView, const FScreenPassTexture& SceneColor, const FScreenPassRenderTarget& Output, bool bIsSceneColorUndistorted = false);
+
+	FScreenPassTexture PostProcessPassSSRInput_RenderThread(FRDGBuilder& GraphBuilder, const FSceneView& InView, const FPostProcessMaterialInputs& Inputs);
+	FScreenPassTexture PostProcessPassAfterTonemap_RenderThread(FRDGBuilder& GraphBuilder, const FSceneView& InView, const FPostProcessMaterialInputs& Inputs);
 
 	// Collection of primitives to render as a custom render pass and composite after post-processing.
 	TSet<TSoftObjectPtr<UPrimitiveComponent>> CompositePrimitives;
@@ -61,5 +68,8 @@ private:
 
 	// Flag to enable global exposure on the composited render
 	std::atomic_bool bCompositeFollowsSceneExposure = false;
+
+	// Flag to enable composite into screen-space reflections
+	std::atomic_bool bCompositeSupportsSSR = false;
 };
 
diff --git a/Engine/Plugins/Experimental/Compositing/HoldoutComposite/Source/HoldoutComposite/Private/HoldoutCompositeSettings.cpp b/Engine/Plugins/Experimental/Compositing/HoldoutComposite/Source/HoldoutComposite/Private/HoldoutCompositeSettings.cpp
--- a/Engine/Plugins/Experimental/Compositing/HoldoutComposite/Source/HoldoutComposite/Private/HoldoutCompositeSettings.cpp
+++ b/Engine/Plugins/Experimental/Compositing/HoldoutComposite/Source/HoldoutComposite/Private/HoldoutCompositeSettings.cpp
@@ -5,7 +5,8 @@
 #include UE_INLINE_GENERATED_CPP_BY_NAME(HoldoutCompositeSettings)
 
 UHoldoutCompositeSettings::UHoldoutCompositeSettings()
-	: bCompositeFollowsSceneExposure(true)
+	: bCompositeFollowsSceneExposure(false)
+	, bCompositeSupportsSSR(true)
 	, SceneViewExtensionPriority(100 + 1 /* OPENCOLORIO_SCENE_VIEW_EXTENSION_PRIORITY + 1 */)
 { }
 
diff --git a/Engine/Plugins/Experimental/Compositing/HoldoutComposite/Source/HoldoutComposite/Public/HoldoutCompositeComponent.h b/Engine/Plugins/Experimental/Compositing/HoldoutComposite/Source/HoldoutComposite/Public/HoldoutCompositeComponent.h
--- a/Engine/Plugins/Experimental/Compositing/HoldoutComposite/Source/HoldoutComposite/Public/HoldoutCompositeComponent.h
+++ b/Engine/Plugins/Experimental/Compositing/HoldoutComposite/Source/HoldoutComposite/Public/HoldoutCompositeComponent.h
@@ -7,7 +7,7 @@
 
 #include "HoldoutCompositeComponent.generated.h"
 
-UCLASS(ClassGroup = Rendering, editinlinenew, meta = (BlueprintSpawnableComponent), MinimalAPI)
+UCLASS(ClassGroup = Rendering, HideCategories=(Activation, Transform, Lighting, Rendering, Tags, Cooking, Physics, LOD, AssetUserData, Navigation), editinlinenew, meta = (BlueprintSpawnableComponent), MinimalAPI)
 class UHoldoutCompositeComponent : public USceneComponent
 {
 	GENERATED_UCLASS_BODY()
diff --git a/Engine/Plugins/Experimental/Compositing/HoldoutComposite/Source/HoldoutComposite/Public/HoldoutCompositeSettings.h b/Engine/Plugins/Experimental/Compositing/HoldoutComposite/Source/HoldoutComposite/Public/HoldoutCompositeSettings.h
--- a/Engine/Plugins/Experimental/Compositing/HoldoutComposite/Source/HoldoutComposite/Public/HoldoutCompositeSettings.h
+++ b/Engine/Plugins/Experimental/Compositing/HoldoutComposite/Source/HoldoutComposite/Public/HoldoutCompositeSettings.h
@@ -33,6 +33,10 @@ public:
 	UPROPERTY(config, EditAnywhere, Category = General)
 	bool bCompositeFollowsSceneExposure;
 
+	/** When enabled, the separate composited render is also used to update screen-space reflections. */
+	UPROPERTY(config, EditAnywhere, Category = General)
+	bool bCompositeSupportsSSR;
+
 	/** Composite (scene view extension) pass priority, which defaults to before OpenColorIO. */
 	UPROPERTY(config, EditAnywhere, Category = General, AdvancedDisplay)
 	int32 SceneViewExtensionPriority;
