HoldoutComposite: Improve toast notifications for missing project settings.

#rnx

diff --git a/Engine/Plugins/Experimental/Compositing/HoldoutComposite/Source/HoldoutComposite/HoldoutComposite.Build.cs b/Engine/Plugins/Experimental/Compositing/HoldoutComposite/Source/HoldoutComposite/HoldoutComposite.Build.cs
--- a/Engine/Plugins/Experimental/Compositing/HoldoutComposite/Source/HoldoutComposite/HoldoutComposite.Build.cs
+++ b/Engine/Plugins/Experimental/Compositing/HoldoutComposite/Source/HoldoutComposite/HoldoutComposite.Build.cs
@@ -43,7 +43,9 @@ public class HoldoutComposite : ModuleRules
 				new string[]
 				{
 					"UnrealEd",
+					"SettingsEditor",
 					"Slate",
+					"SlateCore",
 				}
 			);
 		}
diff --git a/Engine/Plugins/Experimental/Compositing/HoldoutComposite/Source/HoldoutComposite/Private/HoldoutCompositeSubsystem.cpp b/Engine/Plugins/Experimental/Compositing/HoldoutComposite/Source/HoldoutComposite/Private/HoldoutCompositeSubsystem.cpp
--- a/Engine/Plugins/Experimental/Compositing/HoldoutComposite/Source/HoldoutComposite/Private/HoldoutCompositeSubsystem.cpp
+++ b/Engine/Plugins/Experimental/Compositing/HoldoutComposite/Source/HoldoutComposite/Private/HoldoutCompositeSubsystem.cpp
@@ -6,6 +6,16 @@
 
 #include "Engine/RendererSettings.h"
 
+#if WITH_EDITOR
+#include "Framework/Notifications/NotificationManager.h"
+#include "ISettingsEditorModule.h"
+#include "Misc/ConfigCacheIni.h"
+#include "Modules/ModuleManager.h"
+#include "Widgets/Notifications/SNotificationList.h"
+#endif
+
+#define LOCTEXT_NAMESPACE "HoldoutCompositeSubsystem"
+
 UHoldoutCompositeSubsystem::UHoldoutCompositeSubsystem()
 {
 }
@@ -36,19 +46,9 @@ void UHoldoutCompositeSubsystem::RegisterPrimitive(TSoftObjectPtr<UPrimitiveComp
 
 void UHoldoutCompositeSubsystem::RegisterPrimitives(TArrayView<TSoftObjectPtr<UPrimitiveComponent>> InPrimitiveComponents, bool bInHoldoutState)
 {
-	// The compositing relies on alpha preserved through the "tonemapper" post-processing step.
-	static IConsoleVariable* CVarPropagateAlpha = IConsoleManager::Get().FindConsoleVariable(TEXT("r.PostProcessing.PropagateAlpha"));
-	if (CVarPropagateAlpha && CVarPropagateAlpha->GetBool() != true)
-	{
-		UE_CALL_ONCE([]()
-			{
-				UE_LOG(LogHoldoutComposite, Warning, TEXT("Holdout composite is disabled until alpha is enabled through post-processing."));
-			}
-		);
-		return;
-	}
+	const bool bValidSettings = ValidateProjectSettings();
 
-	if (HoldoutCompositeViewExtension.IsValid())
+	if (bValidSettings && HoldoutCompositeViewExtension.IsValid())
 	{
 		HoldoutCompositeViewExtension->RegisterPrimitives(MoveTemp(InPrimitiveComponents), bInHoldoutState);
 	}
@@ -69,3 +69,133 @@ void UHoldoutCompositeSubsystem::UnregisterPrimitives(TArrayView<TSoftObjectPtr<
 	}
 }
 
+bool UHoldoutCompositeSubsystem::ValidateProjectSettings()
+{
+	URendererSettings* RendererSettings = GetMutableDefault<URendererSettings>();
+	check(RendererSettings);
+
+	const bool bValidSettings = RendererSettings->bEnableAlphaChannelInPostProcessing && RendererSettings->bDeferredSupportPrimitiveAlphaHoldout;
+
+	if (!bValidSettings)
+	{
+#if WITH_EDITOR
+		// We inform the user and offer them the option to activate project settings.
+		UE_CALL_ONCE([&]{ PrimitiveHoldoutSettingsNotification(RendererSettings); });
+#else
+		UE_CALL_ONCE([&]{ UE_LOG(LogHoldoutComposite, Warning, TEXT("Both \"Enable Alpha Output\" and \"Support Primitive Alpha Holdout\" project settings must be enabled for holdout composite.")); });
+#endif
+	}
+
+	return bValidSettings;
+}
+
+#if WITH_EDITOR
+void UHoldoutCompositeSubsystem::PrimitiveHoldoutSettingsNotification(class URendererSettings* RendererSettings)
+{
+	const FText HoldoutText = LOCTEXT("HoldoutSettingPrompt", "The following project settings must be enabled for holdout composite:\n- Enable Alpha Output\n- Support Primitive Alpha Holdout");
+	const FText HoldoutConfirmText = LOCTEXT("HoldoutSettingConfirm", "Update");
+	const FText HoldoutCancelText = LOCTEXT("HoldoutSettingCancel", "Not Now");
+
+	/** Utility functions for notifications */
+	struct FSuppressDialogOptions
+	{
+		static bool ShouldSuppressModal()
+		{
+			bool bSuppressNotification = false;
+			GConfig->GetBool(TEXT("HoldoutComposite"), TEXT("SuppressHoldoutCompositePromptNotification"), bSuppressNotification, GEditorPerProjectIni);
+			return bSuppressNotification;
+		}
+
+		static ECheckBoxState GetDontAskAgainCheckBoxState()
+		{
+			return ShouldSuppressModal() ? ECheckBoxState::Checked : ECheckBoxState::Unchecked;
+		}
+
+		static void OnDontAskAgainCheckBoxStateChanged(ECheckBoxState NewState)
+		{
+			// If the user selects to not show this again, set that in the config so we know about it in between sessions
+			const bool bSuppressNotification = (NewState == ECheckBoxState::Checked);
+			GConfig->SetBool(TEXT("HoldoutComposite"), TEXT("SuppressHoldoutCompositePromptNotification"), bSuppressNotification, GEditorPerProjectIni);
+		}
+	};
+
+	FSimpleDelegate EmptyDelegate;
+	FSimpleDelegate OnConfirmDelegate = FSimpleDelegate::CreateLambda(
+		[RendererSettings]()
+		{
+			if (IsValid(RendererSettings))
+			{
+				if (!RendererSettings->bDeferredSupportPrimitiveAlphaHoldout)
+				{
+					RendererSettings->bDeferredSupportPrimitiveAlphaHoldout = true;
+
+					FProperty* Property = RendererSettings->GetClass()->FindPropertyByName(GET_MEMBER_NAME_CHECKED(URendererSettings, bDeferredSupportPrimitiveAlphaHoldout));
+					FPropertyChangedEvent PropertyChangedEvent(Property, EPropertyChangeType::ValueSet, { RendererSettings });
+					RendererSettings->PostEditChangeProperty(PropertyChangedEvent);
+					RendererSettings->UpdateSinglePropertyInConfigFile(Property, RendererSettings->GetDefaultConfigFilename());
+
+					// SupportPrimitiveAlphaHoldout requires shader recompilation, ask for a restart.
+					FModuleManager::GetModuleChecked<ISettingsEditorModule>("SettingsEditor").OnApplicationRestartRequired();
+				}
+
+				if (!RendererSettings->bEnableAlphaChannelInPostProcessing)
+				{
+					RendererSettings->bEnableAlphaChannelInPostProcessing = true;
+
+					FProperty* Property = RendererSettings->GetClass()->FindPropertyByName(GET_MEMBER_NAME_CHECKED(URendererSettings, bEnableAlphaChannelInPostProcessing));
+					FPropertyChangedEvent PropertyChangedEvent(Property, EPropertyChangeType::ValueSet, { RendererSettings });
+					RendererSettings->PostEditChangeProperty(PropertyChangedEvent);
+					RendererSettings->UpdateSinglePropertyInConfigFile(Property, RendererSettings->GetDefaultConfigFilename());
+				}
+			}
+		}
+	);
+
+	FSimpleDelegate OnCancelDelegate = FSimpleDelegate::CreateLambda(
+		[WeakThis = MakeWeakObjectPtr(this)]()
+		{
+			TStrongObjectPtr<UHoldoutCompositeSubsystem> Subsystem = WeakThis.Pin();
+			if (Subsystem.IsValid())
+			{
+				TSharedPtr<SNotificationItem> NotificationItem = Subsystem->HoldoutNotificationItem.Pin();
+				if (NotificationItem.IsValid())
+				{
+					NotificationItem->SetCompletionState(SNotificationItem::CS_None);
+					NotificationItem->ExpireAndFadeout();
+				}
+
+				Subsystem->HoldoutNotificationItem.Reset();
+			}
+		}
+	);
+
+	FNotificationInfo Info(HoldoutText);
+	Info.bFireAndForget = false;
+	Info.bUseLargeFont = false;
+	Info.bUseThrobber = false;
+	Info.bUseSuccessFailIcons = false;
+	Info.ButtonDetails.Add(FNotificationButtonInfo(HoldoutConfirmText, FText(), OnConfirmDelegate));
+	Info.ButtonDetails.Add(FNotificationButtonInfo(HoldoutCancelText, FText(), OnCancelDelegate));
+
+	// Add a "Don't show this again" option
+	Info.CheckBoxState = TAttribute<ECheckBoxState>::Create(&FSuppressDialogOptions::GetDontAskAgainCheckBoxState);
+	Info.CheckBoxStateChanged = FOnCheckStateChanged::CreateStatic(&FSuppressDialogOptions::OnDontAskAgainCheckBoxStateChanged);
+	Info.CheckBoxText = LOCTEXT("DefaultCheckBoxMessage", "Don't show this again");
+
+	if (HoldoutNotificationItem.IsValid())
+	{
+		HoldoutNotificationItem.Pin()->ExpireAndFadeout();
+		HoldoutNotificationItem.Reset();
+	}
+
+	HoldoutNotificationItem = FSlateNotificationManager::Get().AddNotification(Info);
+
+	if (HoldoutNotificationItem.IsValid())
+	{
+		HoldoutNotificationItem.Pin()->SetCompletionState(SNotificationItem::CS_Pending);
+	}
+}
+#endif
+
+#undef LOCTEXT_NAMESPACE
+
diff --git a/Engine/Plugins/Experimental/Compositing/HoldoutComposite/Source/HoldoutComposite/Public/HoldoutCompositeSubsystem.h b/Engine/Plugins/Experimental/Compositing/HoldoutComposite/Source/HoldoutComposite/Public/HoldoutCompositeSubsystem.h
--- a/Engine/Plugins/Experimental/Compositing/HoldoutComposite/Source/HoldoutComposite/Public/HoldoutCompositeSubsystem.h
+++ b/Engine/Plugins/Experimental/Compositing/HoldoutComposite/Source/HoldoutComposite/Public/HoldoutCompositeSubsystem.h
@@ -9,6 +9,8 @@
 #include "HoldoutCompositeSubsystem.generated.h"
 
 class FHoldoutCompositeSceneViewExtension;
+class SNotificationItem;
+class URendererSettings;
 
 /**
  * Composite subsytem used as an interface to the (private) scene view extension.
@@ -41,6 +43,19 @@ public:
 	HOLDOUTCOMPOSITE_API void UnregisterPrimitives(TArrayView<TSoftObjectPtr<UPrimitiveComponent>> InPrimitiveComponents, bool bInHoldoutState=false);
 
 private:
+
+	/* Returns true if the (renderer) project settings are correctly enabled for the composite to be active. */
+	bool ValidateProjectSettings();
+
+#if WITH_EDITOR
+	/* Toast notification to ask users to enable the missing project settings. */
+	void PrimitiveHoldoutSettingsNotification(URendererSettings* RendererSettings);
+
+	/* Toast notification item. */
+	TWeakPtr<SNotificationItem> HoldoutNotificationItem;
+#endif
+
+	/* Owned scene view extension. */
 	TSharedPtr<FHoldoutCompositeSceneViewExtension, ESPMode::ThreadSafe> HoldoutCompositeViewExtension;
 };
 
