Holdout: Add enable state getter & setter on holdout component.

#rnx

diff --git a/Engine/Plugins/Experimental/Compositing/HoldoutComposite/Source/HoldoutComposite/Private/HoldoutCompositeComponent.cpp b/Engine/Plugins/Experimental/Compositing/HoldoutComposite/Source/HoldoutComposite/Private/HoldoutCompositeComponent.cpp
--- a/Engine/Plugins/Experimental/Compositing/HoldoutComposite/Source/HoldoutComposite/Private/HoldoutCompositeComponent.cpp
+++ b/Engine/Plugins/Experimental/Compositing/HoldoutComposite/Source/HoldoutComposite/Private/HoldoutCompositeComponent.cpp
@@ -122,12 +122,37 @@ void UHoldoutCompositeComponent::OnAttachmentChanged()
 			RegisterCompositeImpl();
 		}
 	}
+}
+
+bool UHoldoutCompositeComponent::IsEnabled() const
+{
+	return bIsEnabled;
+}
 
+void UHoldoutCompositeComponent::SetEnabled(bool bInIsEnabled)
+{
+	if (bIsEnabled != bInIsEnabled)
+	{
+		bIsEnabled = bInIsEnabled;
 
+		if (bIsEnabled)
+		{
+			RegisterCompositeImpl();
+		}
+		else
+		{
+			UnregisterCompositeImpl();
+		}
+	}
 }
 
 void UHoldoutCompositeComponent::RegisterCompositeImpl()
 {
+	if (!bIsEnabled)
+	{
+		return;
+	}
+
 	UHoldoutCompositeSubsystem* Subsystem = UWorld::GetSubsystem<UHoldoutCompositeSubsystem>(GetWorld());
 	TArray<TSoftObjectPtr<UPrimitiveComponent>> ParentPrimitives = FindPrimitiveComponents(GetAttachParent());
 
diff --git a/Engine/Plugins/Experimental/Compositing/HoldoutComposite/Source/HoldoutComposite/Public/HoldoutCompositeComponent.h b/Engine/Plugins/Experimental/Compositing/HoldoutComposite/Source/HoldoutComposite/Public/HoldoutCompositeComponent.h
--- a/Engine/Plugins/Experimental/Compositing/HoldoutComposite/Source/HoldoutComposite/Public/HoldoutCompositeComponent.h
+++ b/Engine/Plugins/Experimental/Compositing/HoldoutComposite/Source/HoldoutComposite/Public/HoldoutCompositeComponent.h
@@ -23,6 +23,14 @@ public:
 	HOLDOUTCOMPOSITE_API virtual void DetachFromComponent(const FDetachmentTransformRules& DetachmentRules) override;
 	//~ End USceneComponent Interface
 
+	/* Get the enabled state of the component. */
+	UFUNCTION(BlueprintGetter)
+	bool IsEnabled() const;
+
+	/* Set the enabled state of the component. */
+	UFUNCTION(BlueprintSetter)
+	void SetEnabled(bool bInEnabled);
+
 private:
 
 	/* Private implementation of the register method. */
@@ -31,8 +39,10 @@ private:
 	/* Private implementation of the unregister method. */
 	void UnregisterCompositeImpl();
 
-	/* Primitive holdout property value of the parent component. */
-	UPROPERTY()
-	TOptional<bool> bCachedParentHoldout = false;
+private:
+
+	/* Whether or not the component activates the composite. */
+	UPROPERTY(EditAnywhere, BlueprintGetter = IsEnabled, BlueprintSetter = SetEnabled, Category = "HoldoutComposite", meta = (AllowPrivateAccess = true))
+	bool bIsEnabled = true;
 };
 
