HoldoutComposite: Minor improvement & fixes just-added toast notification, copied back from MRG-approved implementation.

#rnx

diff --git a/Engine/Plugins/Experimental/Compositing/HoldoutComposite/Source/HoldoutComposite/Private/HoldoutCompositeSubsystem.cpp b/Engine/Plugins/Experimental/Compositing/HoldoutComposite/Source/HoldoutComposite/Private/HoldoutCompositeSubsystem.cpp
--- a/Engine/Plugins/Experimental/Compositing/HoldoutComposite/Source/HoldoutComposite/Private/HoldoutCompositeSubsystem.cpp
+++ b/Engine/Plugins/Experimental/Compositing/HoldoutComposite/Source/HoldoutComposite/Private/HoldoutCompositeSubsystem.cpp
@@ -90,10 +90,15 @@ bool UHoldoutCompositeSubsystem::ValidateProjectSettings()
 }
 
 #if WITH_EDITOR
-void UHoldoutCompositeSubsystem::PrimitiveHoldoutSettingsNotification(class URendererSettings* RendererSettings)
+void UHoldoutCompositeSubsystem::PrimitiveHoldoutSettingsNotification(URendererSettings* RendererSettings)
 {
-	const FText HoldoutText = LOCTEXT("HoldoutSettingPrompt", "The following project settings must be enabled for holdout composite:\n- Alpha Output\n- Support Primitive Alpha Holdout");
-	const FText HoldoutConfirmText = LOCTEXT("HoldoutSettingConfirm", "Update");
+	const bool bAlphaOutputMissing = !RendererSettings->bEnableAlphaChannelInPostProcessing;
+	const bool bPrimitiveHoldoutMissing = !RendererSettings->bDeferredSupportPrimitiveAlphaHoldout;
+
+	const FText AlphaOutputSettingOption = LOCTEXT("HoldoutSetting_AlphaOutput", "\n- Alpha Output");
+	const FText PrimitiveHoldoutSettingOption = LOCTEXT("HoldoutSetting_PrimitiveHoldout", "\n- Support Primitive Alpha Holdout");
+	const FText HoldoutText = FText::Format(LOCTEXT("HoldoutSettingPrompt", "The following project setting(s) must be enabled for holdout composite:{0}{1}"), bAlphaOutputMissing ? AlphaOutputSettingOption : FText::GetEmpty(), bPrimitiveHoldoutMissing ? PrimitiveHoldoutSettingOption : FText::GetEmpty());
+	const FText HoldoutConfirmText = LOCTEXT("HoldoutSettingConfirm", "Enable");
 	const FText HoldoutCancelText = LOCTEXT("HoldoutSettingCancel", "Not Now");
 
 	/** Utility functions for notifications */
@@ -119,7 +124,12 @@ void UHoldoutCompositeSubsystem::PrimitiveHoldoutSettingsNotification(class URen
 		}
 	};
 
-	FSimpleDelegate EmptyDelegate;
+	// If the user has specified to supress this pop up, then just early out and exit	
+	if (FSuppressDialogOptions::ShouldSuppressModal())
+	{
+		return;
+	}
+
 	FSimpleDelegate OnConfirmDelegate = FSimpleDelegate::CreateLambda(
 		[RendererSettings]()
 		{
@@ -127,9 +137,11 @@ void UHoldoutCompositeSubsystem::PrimitiveHoldoutSettingsNotification(class URen
 			{
 				if (!RendererSettings->bDeferredSupportPrimitiveAlphaHoldout)
 				{
+					FProperty* Property = RendererSettings->GetClass()->FindPropertyByName(GET_MEMBER_NAME_CHECKED(URendererSettings, bDeferredSupportPrimitiveAlphaHoldout));
+					RendererSettings->PreEditChange(Property);
+
 					RendererSettings->bDeferredSupportPrimitiveAlphaHoldout = true;
 
-					FProperty* Property = RendererSettings->GetClass()->FindPropertyByName(GET_MEMBER_NAME_CHECKED(URendererSettings, bDeferredSupportPrimitiveAlphaHoldout));
 					FPropertyChangedEvent PropertyChangedEvent(Property, EPropertyChangeType::ValueSet, { RendererSettings });
 					RendererSettings->PostEditChangeProperty(PropertyChangedEvent);
 					RendererSettings->UpdateSinglePropertyInConfigFile(Property, RendererSettings->GetDefaultConfigFilename());
@@ -140,9 +152,11 @@ void UHoldoutCompositeSubsystem::PrimitiveHoldoutSettingsNotification(class URen
 
 				if (!RendererSettings->bEnableAlphaChannelInPostProcessing)
 				{
+					FProperty* Property = RendererSettings->GetClass()->FindPropertyByName(GET_MEMBER_NAME_CHECKED(URendererSettings, bEnableAlphaChannelInPostProcessing));
+					RendererSettings->PreEditChange(Property);
+
 					RendererSettings->bEnableAlphaChannelInPostProcessing = true;
 
-					FProperty* Property = RendererSettings->GetClass()->FindPropertyByName(GET_MEMBER_NAME_CHECKED(URendererSettings, bEnableAlphaChannelInPostProcessing));
 					FPropertyChangedEvent PropertyChangedEvent(Property, EPropertyChangeType::ValueSet, { RendererSettings });
 					RendererSettings->PostEditChangeProperty(PropertyChangedEvent);
 					RendererSettings->UpdateSinglePropertyInConfigFile(Property, RendererSettings->GetDefaultConfigFilename());
@@ -180,7 +194,7 @@ void UHoldoutCompositeSubsystem::PrimitiveHoldoutSettingsNotification(class URen
 	// Add a "Don't show this again" option
 	Info.CheckBoxState = TAttribute<ECheckBoxState>::Create(&FSuppressDialogOptions::GetDontAskAgainCheckBoxState);
 	Info.CheckBoxStateChanged = FOnCheckStateChanged::CreateStatic(&FSuppressDialogOptions::OnDontAskAgainCheckBoxStateChanged);
-	Info.CheckBoxText = LOCTEXT("DefaultCheckBoxMessage", "Don't show this again");
+	Info.CheckBoxText = LOCTEXT("DontShowThisAgainCheckBoxMessage", "Don't show this again");
 
 	if (HoldoutNotificationItem.IsValid())
 	{
