HoldoutComposite: Partial fix for holdout composites inside (user) scene captures.

#rnx

diff --git a/Engine/Plugins/Experimental/Compositing/HoldoutComposite/Source/HoldoutComposite/Private/HoldoutCompositeSceneViewExtension.cpp b/Engine/Plugins/Experimental/Compositing/HoldoutComposite/Source/HoldoutComposite/Private/HoldoutCompositeSceneViewExtension.cpp
--- a/Engine/Plugins/Experimental/Compositing/HoldoutComposite/Source/HoldoutComposite/Private/HoldoutCompositeSceneViewExtension.cpp
+++ b/Engine/Plugins/Experimental/Compositing/HoldoutComposite/Source/HoldoutComposite/Private/HoldoutCompositeSceneViewExtension.cpp
@@ -329,7 +329,8 @@ void FHoldoutCompositeSceneViewExtension::SetupView(FSceneViewFamily& InViewFami
 	);
 
 	FSceneInterface::FCustomRenderPassRendererInput PassInput;
-	PassInput.ViewLocation = InView.ViewLocation;
+	// Note: Incoming view location is invalid for scene captures
+	PassInput.ViewLocation = InView.bIsSceneCapture ? InView.ViewMatrices.GetViewOrigin() : InView.ViewLocation;
 	PassInput.ViewRotationMatrix = InView.ViewMatrices.GetViewMatrix().RemoveTranslation();
 	PassInput.ViewRotationMatrix.RemoveScaling();
 
