Holdout: Fixed missing expire and fade out calls on primitive alpha holdout setting update confirmation.

#rnx

diff --git a/Engine/Plugins/Experimental/Compositing/HoldoutComposite/Source/HoldoutComposite/Private/HoldoutCompositeSubsystem.cpp b/Engine/Plugins/Experimental/Compositing/HoldoutComposite/Source/HoldoutComposite/Private/HoldoutCompositeSubsystem.cpp
--- a/Engine/Plugins/Experimental/Compositing/HoldoutComposite/Source/HoldoutComposite/Private/HoldoutCompositeSubsystem.cpp
+++ b/Engine/Plugins/Experimental/Compositing/HoldoutComposite/Source/HoldoutComposite/Private/HoldoutCompositeSubsystem.cpp
@@ -131,7 +131,7 @@ void UHoldoutCompositeSubsystem::PrimitiveHoldoutSettingsNotification(URendererS
 	}
 
 	FSimpleDelegate OnConfirmDelegate = FSimpleDelegate::CreateLambda(
-		[RendererSettings]()
+		[WeakThis = MakeWeakObjectPtr(this), RendererSettings]()
 		{
 			if (IsValid(RendererSettings))
 			{
@@ -162,6 +162,19 @@ void UHoldoutCompositeSubsystem::PrimitiveHoldoutSettingsNotification(URendererS
 					RendererSettings->UpdateSinglePropertyInConfigFile(Property, RendererSettings->GetDefaultConfigFilename());
 				}
 			}
+
+			TStrongObjectPtr<UHoldoutCompositeSubsystem> Subsystem = WeakThis.Pin();
+			if (Subsystem.IsValid())
+			{
+				TSharedPtr<SNotificationItem> NotificationItem = Subsystem->HoldoutNotificationItem.Pin();
+				if (NotificationItem.IsValid())
+				{
+					NotificationItem->SetCompletionState(SNotificationItem::CS_Success);
+					NotificationItem->ExpireAndFadeout();
+				}
+
+				Subsystem->HoldoutNotificationItem.Reset();
+			}
 		}
 	);
 
diff --git a/Engine/Plugins/MovieScene/MovieRenderPipeline/Source/MovieRenderPipelineCore/Private/Graph/MovieGraphUtils.cpp b/Engine/Plugins/MovieScene/MovieRenderPipeline/Source/MovieRenderPipelineCore/Private/Graph/MovieGraphUtils.cpp
--- a/Engine/Plugins/MovieScene/MovieRenderPipeline/Source/MovieRenderPipelineCore/Private/Graph/MovieGraphUtils.cpp
+++ b/Engine/Plugins/MovieScene/MovieRenderPipeline/Source/MovieRenderPipelineCore/Private/Graph/MovieGraphUtils.cpp
@@ -157,6 +157,14 @@ namespace UE::MovieGraph
 							// SupportPrimitiveAlphaHoldout requires shader recompilation, ask for a restart.
 							FModuleManager::GetModuleChecked<ISettingsEditorModule>("SettingsEditor").OnApplicationRestartRequired();
 						}
+
+						if (TSharedPtr<SNotificationItem> Item = NotificationItem.Pin())
+						{
+							Item->SetCompletionState(SNotificationItem::CS_Success);
+							Item->ExpireAndFadeout();
+						}
+
+						NotificationItem.Reset();
 					}
 				}
 			);
