HoldoutComposite: Removed workaround if-def since CRP approach has now stabilitized.

diff --git a/Engine/Plugins/Experimental/Compositing/HoldoutComposite/Source/HoldoutComposite/Private/HoldoutCompositeSceneViewExtension.cpp b/Engine/Plugins/Experimental/Compositing/HoldoutComposite/Source/HoldoutComposite/Private/HoldoutCompositeSceneViewExtension.cpp
--- a/Engine/Plugins/Experimental/Compositing/HoldoutComposite/Source/HoldoutComposite/Private/HoldoutCompositeSceneViewExtension.cpp
+++ b/Engine/Plugins/Experimental/Compositing/HoldoutComposite/Source/HoldoutComposite/Private/HoldoutCompositeSceneViewExtension.cpp
@@ -147,11 +147,7 @@ public:
 	{
 		FRDGTextureRef DilatedCRP = CreateDilatedTexture(GraphBuilder);
 
-#if HOLDOUT_COMPOSITE_WORKAROUND_UE_209928
 		ParentExtension->CollectCustomRenderTarget(ViewId, GraphBuilder.ConvertToExternalTexture(DilatedCRP));
-#else
-		ParentExtension->CollectCustomRenderTarget(ViewId, DilatedCRP);
-#endif
 	}
 
 private:
@@ -341,7 +337,6 @@ void FHoldoutCompositeSceneViewExtension::SetupView(FSceneViewFamily& InViewFami
 	PassInput.CustomRenderPass = CustomRenderPass;
 	PassInput.bIsSceneCapture = true;
 
-	// TODO: Once CRPs are associated with view(family), disable & remove HOLDOUT_COMPOSITE_WORKAROUND_UE_209928.
 	WorldPtr.Get()->Scene->AddCustomRenderPass(&InViewFamily, PassInput);
 }
 
@@ -367,19 +362,11 @@ FRDGTextureRef FHoldoutCompositeSceneViewExtension::GetCustomRenderPassTexture(F
 {
 	FRDGTextureRef CustomRenderPassTexture = GSystemTextures.GetBlackAlphaOneDummy(GraphBuilder);
 
-#if HOLDOUT_COMPOSITE_WORKAROUND_UE_209928
 	const TRefCountPtr<IPooledRenderTarget>* CustomRenderPassRenderTargetPtr = CustomRenderTargetPerView_RenderThread.Find(InView.GetViewKey());
 	if (CustomRenderPassRenderTargetPtr != nullptr)
 	{
 		CustomRenderPassTexture = GraphBuilder.RegisterExternalTexture(*CustomRenderPassRenderTargetPtr);
 	}
-#else
-	const FRDGTextureRef* CustomRenderPassTexturePtr = CustomRenderTargetPerView_RenderThread.Find(InView.GetViewKey());
-	if (CustomRenderPassTexturePtr && HasBeenProduced(*CustomRenderPassTexturePtr))
-	{
-		CustomRenderPassTexture = *CustomRenderPassTexturePtr;
-	}
-#endif
 	
 	return CustomRenderPassTexture;
 }
diff --git a/Engine/Plugins/Experimental/Compositing/HoldoutComposite/Source/HoldoutComposite/Private/HoldoutCompositeSceneViewExtension.h b/Engine/Plugins/Experimental/Compositing/HoldoutComposite/Source/HoldoutComposite/Private/HoldoutCompositeSceneViewExtension.h
--- a/Engine/Plugins/Experimental/Compositing/HoldoutComposite/Source/HoldoutComposite/Private/HoldoutCompositeSceneViewExtension.h
+++ b/Engine/Plugins/Experimental/Compositing/HoldoutComposite/Source/HoldoutComposite/Private/HoldoutCompositeSceneViewExtension.h
@@ -7,9 +7,6 @@
 #include "RendererInterface.h"
 #include "Containers/ContainersFwd.h"
 
-// Temporary workaround for UE-209928 to fix holdout composite on nDisplay. Remove-disable once addressed.
-#define HOLDOUT_COMPOSITE_WORKAROUND_UE_209928 1
-
 class FViewInfo;
 struct FScreenPassRenderTarget;
 struct FScreenPassTexture;
@@ -59,12 +56,8 @@ private:
 	// Collection of primitives to render as a custom render pass and composite after post-processing.
 	TSet<TSoftObjectPtr<UPrimitiveComponent>> CompositePrimitives;
 
-#if HOLDOUT_COMPOSITE_WORKAROUND_UE_209928
-	TMap<uint32, TRefCountPtr<IPooledRenderTarget>> CustomRenderTargetPerView_RenderThread;
-#else
 	// Custom render pass render targets for each active view
-	TMap<uint32, FRDGTextureRef> CustomRenderTargetPerView_RenderThread;
-#endif
+	TMap<uint32, TRefCountPtr<IPooledRenderTarget>> CustomRenderTargetPerView_RenderThread;
 
 	// Flag to enable global exposure on the composited render
 	std::atomic_bool bCompositeFollowsSceneExposure = false;
