HoldoutComposite: Adds support for translucent materials.

We also automatically route the Base Color output to Emissive when bUnlitViewmode is active.This is a debug non-shipping PC-only feature, but it is automatically enabled due to its convenience.

diff --git a/Engine/Plugins/Experimental/Compositing/HoldoutComposite/Shaders/Private/HoldoutCompositeFinal.usf b/Engine/Plugins/Experimental/Compositing/HoldoutComposite/Shaders/Private/HoldoutCompositeFinal.usf
--- a/Engine/Plugins/Experimental/Compositing/HoldoutComposite/Shaders/Private/HoldoutCompositeFinal.usf
+++ b/Engine/Plugins/Experimental/Compositing/HoldoutComposite/Shaders/Private/HoldoutCompositeFinal.usf
@@ -25,6 +25,21 @@ StructuredBuffer<float4> EyeAdaptationBuffer;
 uint2 Encodings;
 float2 DisplayGamma;
 
+float SafeRcp(float x)
+{
+	return x > 0.0 ? rcp(x) : 0.0;
+}
+
+float4 InvertAlpha(float4 InColor)
+{
+	return float4(InColor.xyz, 1.0 - InColor.w);
+}
+
+float4 Unpremultiply(float4 InColor)
+{
+	return float4(InColor.xyz * SafeRcp(InColor.w), InColor.w);
+}
+
 void MainPS(
     float4 SvPosition : SV_POSITION,
     out float4 OutColor : SV_Target0
@@ -54,8 +69,11 @@ void MainPS(
         SceneColor.rgb = sRGBToLinear(SceneColor.rgb);
     }
 
+    // The alpha comes from the scene color holdout mask, so we unpremultiply the custom render to get its opaque colors. (No-op on opaque geometry.)
+    CustomColor = InvertAlpha(Unpremultiply(InvertAlpha(CustomColor)));
+
     // TODO: HDR cases will need to be handled explicitely.
-    OutColor.rgb = CustomColor.rgb * SceneColor.a + SceneColor.rgb;
+    OutColor.rgb = SceneColor.rgb + CustomColor.rgb * SceneColor.a;
     OutColor.a = CustomColor.a * SceneColor.a;
 
     BRANCH
diff --git a/Engine/Plugins/Experimental/Compositing/HoldoutComposite/Source/HoldoutComposite/HoldoutComposite.Build.cs b/Engine/Plugins/Experimental/Compositing/HoldoutComposite/Source/HoldoutComposite/HoldoutComposite.Build.cs
--- a/Engine/Plugins/Experimental/Compositing/HoldoutComposite/Source/HoldoutComposite/HoldoutComposite.Build.cs
+++ b/Engine/Plugins/Experimental/Compositing/HoldoutComposite/Source/HoldoutComposite/HoldoutComposite.Build.cs
@@ -11,12 +11,6 @@ public class HoldoutComposite : ModuleRules
 		bTreatAsEngineModule = true;
 		PCHUsage = PCHUsageMode.UseExplicitOrSharedPCHs;
 
-		PrivateIncludePaths.AddRange(
-			new string[] {
-				System.IO.Path.Combine(GetModuleDirectory("Renderer"), "Private"),
-			}
-		);
-
 		PublicDependencyModuleNames.AddRange(
 			new string[]
 			{
diff --git a/Engine/Plugins/Experimental/Compositing/HoldoutComposite/Source/HoldoutComposite/Private/HoldoutCompositeSceneViewExtension.cpp b/Engine/Plugins/Experimental/Compositing/HoldoutComposite/Source/HoldoutComposite/Private/HoldoutCompositeSceneViewExtension.cpp
--- a/Engine/Plugins/Experimental/Compositing/HoldoutComposite/Source/HoldoutComposite/Private/HoldoutCompositeSceneViewExtension.cpp
+++ b/Engine/Plugins/Experimental/Compositing/HoldoutComposite/Source/HoldoutComposite/Private/HoldoutCompositeSceneViewExtension.cpp
@@ -21,9 +21,6 @@
 #include "ScreenPass.h"
 #include "SystemTextures.h"
 
-// Private renderer includes currently needed...
-#include "SceneRendering.h"
-
 DECLARE_GPU_STAT_NAMED(FHoldoutCompositeDilate, TEXT("HoldoutComposite.Dilate"));
 DECLARE_GPU_STAT_NAMED(FHoldoutCompositeSSRInput, TEXT("HoldoutComposite.SSRInput"));
 DECLARE_GPU_STAT_NAMED(FHoldoutCompositeFinal, TEXT("HoldoutComposite.Final"));
@@ -117,27 +114,11 @@ public:
 		, ViewId(InView.GetViewKey())
 		, ViewFeatureLevel(InView.GetFeatureLevel())
 	{
+		bSceneColorWithTranslucent = true;
 	}
 
 	virtual void OnPreRender(FRDGBuilder& GraphBuilder) override
 	{
-		/**
-		 * Note: We abuse the reflection capture view property in a custom render pass to disable
-		 * primitive alpha holdout during its base pass render. Because the holdout is part of the
-		 * primitive uniform buffer, it cannot easily have true/false states in the same frame.
-		 * We would otherwise need to duplicate the entire primitive so that the main render has
-		 * holdout enabled, while the custom render pass primitive has it disabled.
-		 * 
-		 * Maintaining these duplicates becomes problematic, especially with different types of meshes.
-		 * As a result, we consider this reflection capture property override preferable.
-		*/
-		for (FViewInfo* ViewInfo : Views)
-		{
-			// Holdout is ignored during reflections captures to preserve indirect light (see PRIMITIVE_SCENE_DATA_FLAG_HOLDOUT in BasePassPixelShader.usf)
-			ViewInfo->CachedViewUniformShaderParameters->RenderingReflectionCaptureMask = 1.0f;
-			ViewInfo->ViewUniformBuffer.UpdateUniformBufferImmediate(GraphBuilder.RHICmdList, *ViewInfo->CachedViewUniformShaderParameters);
-		}
-
 		const FRDGTextureDesc TextureDesc = FRDGTextureDesc::Create2D(RenderTargetSize, PF_FloatRGBA, FClearValueBinding::Black, TexCreate_RenderTargetable | TexCreate_ShaderResource);
 		RenderTargetTexture = GraphBuilder.CreateTexture(TextureDesc, TEXT("HoldoutCompositeCustomTexture"));
 		AddClearRenderTargetPass(GraphBuilder, RenderTargetTexture, FLinearColor::Black, FIntRect(FInt32Point(), RenderTargetSize));
@@ -145,28 +126,28 @@ public:
 
 	virtual void OnPostRender(FRDGBuilder& GraphBuilder) override
 	{
-		FRDGTextureRef DilatedCRP = CreateDilatedTexture(GraphBuilder);
+		FRDGTextureRef DilatedCRP = CreateDilatedTexture(GraphBuilder, RenderTargetTexture, ViewFeatureLevel);
 
 		ParentExtension->CollectCustomRenderTarget(ViewId, GraphBuilder.ConvertToExternalTexture(DilatedCRP));
 	}
 
 private:
 	
-	FRDGTextureRef CreateDilatedTexture(FRDGBuilder& GraphBuilder)
+	static FRDGTextureRef CreateDilatedTexture(FRDGBuilder& GraphBuilder, FRDGTextureRef InRenderTarget, ERHIFeatureLevel::Type FeatureLevel)
 	{
 		RDG_EVENT_SCOPE_STAT(GraphBuilder, FHoldoutCompositeDilate, "HoldoutComposite.Dilate");
 		RDG_GPU_STAT_SCOPE(GraphBuilder, FHoldoutCompositeDilate);
 
-		FGlobalShaderMap* GlobalShaderMap = GetGlobalShaderMap(ViewFeatureLevel);
+		FGlobalShaderMap* GlobalShaderMap = GetGlobalShaderMap(FeatureLevel);
 
-		const FIntPoint TextureSize = FIntPoint(RenderTargetTexture->Desc.GetSize().X, RenderTargetTexture->Desc.GetSize().Y);
+		const FIntPoint TextureSize = FIntPoint(InRenderTarget->Desc.GetSize().X, InRenderTarget->Desc.GetSize().Y);
 		const FRDGTextureDesc TextureDesc = FRDGTextureDesc::Create2D(TextureSize, PF_FloatRGBA, FClearValueBinding::Black, TexCreate_UAV | TexCreate_ShaderResource);
 		FRDGTextureRef DilatedTexture = GraphBuilder.CreateTexture(TextureDesc, TEXT("HoldoutCompositeDilatedTexture"));
 
 		// Async compute dilation pass
 		{
 			FHoldoutCompositeDilateShader::FParameters* PassParameters = GraphBuilder.AllocParameters<FHoldoutCompositeDilateShader::FParameters>();
-			PassParameters->InputTexture = RenderTargetTexture;
+			PassParameters->InputTexture = InRenderTarget;
 			PassParameters->RWOutputTexture = GraphBuilder.CreateUAV(DilatedTexture);
 			PassParameters->Dimensions = TextureSize;
 
@@ -324,7 +305,15 @@ void FHoldoutCompositeSceneViewExtension::SetupView(FSceneViewFamily& InViewFami
 		InView
 	);
 
-	FSceneInterface::FCustomRenderPassRendererInput PassInput;
+	const UHoldoutCompositeSettings* Settings = GetDefault<UHoldoutCompositeSettings>();
+
+	FSceneInterface::FCustomRenderPassRendererInput PassInput{};
+	PassInput.EngineShowFlags = InViewFamily.EngineShowFlags;
+	PassInput.EngineShowFlags.DisableFeaturesForUnlit();
+	PassInput.EngineShowFlags.SetTranslucency(true);
+	PassInput.EngineShowFlags.SetUnlitViewmode(Settings ? Settings->bUnlitViewmode : false);
+	PassInput.EngineShowFlags.SetAllowPrimitiveAlphaHoldout(false);
+
 	// Note: Incoming view location is invalid for scene captures
 	PassInput.ViewLocation = InView.bIsSceneCapture ? InView.ViewMatrices.GetViewOrigin() : InView.ViewLocation;
 	PassInput.ViewRotationMatrix = InView.ViewMatrices.GetViewMatrix().RemoveTranslation();
diff --git a/Engine/Plugins/Experimental/Compositing/HoldoutComposite/Source/HoldoutComposite/Private/HoldoutCompositeSettings.cpp b/Engine/Plugins/Experimental/Compositing/HoldoutComposite/Source/HoldoutComposite/Private/HoldoutCompositeSettings.cpp
--- a/Engine/Plugins/Experimental/Compositing/HoldoutComposite/Source/HoldoutComposite/Private/HoldoutCompositeSettings.cpp
+++ b/Engine/Plugins/Experimental/Compositing/HoldoutComposite/Source/HoldoutComposite/Private/HoldoutCompositeSettings.cpp
@@ -7,6 +7,7 @@
 UHoldoutCompositeSettings::UHoldoutCompositeSettings()
 	: bCompositeFollowsSceneExposure(false)
 	, bCompositeSupportsSSR(true)
+	, bUnlitViewmode(true)
 	, SceneViewExtensionPriority(100 + 1 /* OPENCOLORIO_SCENE_VIEW_EXTENSION_PRIORITY + 1 */)
 { }
 
diff --git a/Engine/Plugins/Experimental/Compositing/HoldoutComposite/Source/HoldoutComposite/Public/HoldoutCompositeSettings.h b/Engine/Plugins/Experimental/Compositing/HoldoutComposite/Source/HoldoutComposite/Public/HoldoutCompositeSettings.h
--- a/Engine/Plugins/Experimental/Compositing/HoldoutComposite/Source/HoldoutComposite/Public/HoldoutCompositeSettings.h
+++ b/Engine/Plugins/Experimental/Compositing/HoldoutComposite/Source/HoldoutComposite/Public/HoldoutCompositeSettings.h
@@ -37,6 +37,10 @@ public:
 	UPROPERTY(config, EditAnywhere, Category = General)
 	bool bCompositeSupportsSSR;
 
+	/** Enables the development shader debug feature that routes the Base Color output to Emissive for the separate render. Non-shipping PC build only. */
+	UPROPERTY(config, EditAnywhere, Category = General)
+	bool bUnlitViewmode;
+
 	/** Composite (scene view extension) pass priority, which defaults to before OpenColorIO. */
 	UPROPERTY(config, EditAnywhere, Category = General, AdvancedDisplay)
 	int32 SceneViewExtensionPriority;
