HoldoutComposite: Adjusted previous CL's alpha treshold.

#rnx

diff --git a/Engine/Plugins/Experimental/Compositing/HoldoutComposite/Shaders/Private/HoldoutCompositeDilate.usf b/Engine/Plugins/Experimental/Compositing/HoldoutComposite/Shaders/Private/HoldoutCompositeDilate.usf
--- a/Engine/Plugins/Experimental/Compositing/HoldoutComposite/Shaders/Private/HoldoutCompositeDilate.usf
+++ b/Engine/Plugins/Experimental/Compositing/HoldoutComposite/Shaders/Private/HoldoutCompositeDilate.usf
@@ -7,7 +7,7 @@
 #endif
 
 #define KERNEL_COUNT 8
-#define ALPHA_THRESHOLD (1.0f - 1e-9f)
+#define ALPHA_THRESHOLD (1.0f - 1e-7f)
 
 // Number of texels arround the group tile for scene color.
 #define LDS_TILE_BORDER_SIZE (1)
