HoldoutComposite: Additional plugin settings allowing users to disable lens distortion on the separate render.

* Application of lens distortion from TSR can now be enabled/disabled.
* While the composite pass runs by default after tonemap, it can now be moved at the end of post processing for use cases where the comp should occur last (for example after the lens distortion PPM).
* Cleaned up plugin settings.

diff --git a/Engine/Plugins/Experimental/Compositing/HoldoutComposite/Shaders/Private/HoldoutCompositeFinal.usf b/Engine/Plugins/Experimental/Compositing/HoldoutComposite/Shaders/Private/HoldoutCompositeFinal.usf
--- a/Engine/Plugins/Experimental/Compositing/HoldoutComposite/Shaders/Private/HoldoutCompositeFinal.usf
+++ b/Engine/Plugins/Experimental/Compositing/HoldoutComposite/Shaders/Private/HoldoutCompositeFinal.usf
@@ -24,6 +24,7 @@ SamplerState UndistortingDisplacementSampler;
 StructuredBuffer<float4> EyeAdaptationBuffer;
 uint2 Encodings;
 float2 DisplayGamma;
+uint bApplyTSRLensDistortion;
 
 float SafeRcp(float x)
 {
@@ -46,12 +47,14 @@ void MainPS(
     )
 {
     float2 DistortedViewportUV = (SvPosition.xy - Output_ViewportMin) * Output_ViewportSizeInverse;
-    const float2 UndistortedViewportUV = ApplyLensDistortionOnViewportUV(UndistortingDisplacementTexture, UndistortingDisplacementSampler, DistortedViewportUV);
     
     const float2 InputUV = (DistortedViewportUV * Input_ViewportSize + Input_ViewportMin) * Input_ExtentInverse;
     float4 SceneColor = Texture2DSample(InputTexture, InputSampler, InputUV);
 
-    const float2 CustomUV = (UndistortedViewportUV * Custom_ViewportSize + Custom_ViewportMin) * Custom_ExtentInverse;
+    const float2 UndistortedViewportUV = ApplyLensDistortionOnViewportUV(UndistortingDisplacementTexture, UndistortingDisplacementSampler, DistortedViewportUV);
+    const float2 CustomViewportUV = select(bApplyTSRLensDistortion > 0, UndistortedViewportUV, DistortedViewportUV);
+
+    const float2 CustomUV = (CustomViewportUV * Custom_ViewportSize + Custom_ViewportMin) * Custom_ExtentInverse;
     float4 CustomColor = Texture2DSample(CustomTexture, CustomSampler, CustomUV);
 
 #if USE_GLOBAL_EXPOSURE
diff --git a/Engine/Plugins/Experimental/Compositing/HoldoutComposite/Source/HoldoutComposite/Private/HoldoutCompositeSceneViewExtension.cpp b/Engine/Plugins/Experimental/Compositing/HoldoutComposite/Source/HoldoutComposite/Private/HoldoutCompositeSceneViewExtension.cpp
--- a/Engine/Plugins/Experimental/Compositing/HoldoutComposite/Source/HoldoutComposite/Private/HoldoutCompositeSceneViewExtension.cpp
+++ b/Engine/Plugins/Experimental/Compositing/HoldoutComposite/Source/HoldoutComposite/Private/HoldoutCompositeSceneViewExtension.cpp
@@ -98,6 +98,7 @@ class FHoldoutCompositeFinalShader : public FGlobalShader
 		SHADER_PARAMETER_RDG_BUFFER_SRV(StructuredBuffer<float4>, EyeAdaptationBuffer)
 		SHADER_PARAMETER(FUint32Vector2, Encodings)
 		SHADER_PARAMETER(FVector2f, DisplayGamma)
+		SHADER_PARAMETER(uint32, bApplyTSRLensDistortion)
 		RENDER_TARGET_BINDING_SLOTS()
 	END_SHADER_PARAMETER_STRUCT()
 };
@@ -172,6 +173,10 @@ private:
 
 FHoldoutCompositeSceneViewExtension::FHoldoutCompositeSceneViewExtension(const FAutoRegister& AutoReg, UWorld* InWorld)
 	: FWorldSceneViewExtension(AutoReg, InWorld)
+	, bFollowSceneExposure(false)
+	, bSupportSSR(true)
+	, bApplyTSRLensDistortion(true)
+	, CompositePassLocation(EHoldoutCompositePassLocation::AfterTonemap)
 {
 }
 
@@ -233,12 +238,8 @@ bool FHoldoutCompositeSceneViewExtension::IsActiveThisFrame_Internal(const FScen
 int32 FHoldoutCompositeSceneViewExtension::GetPriority() const
 {
 	const UHoldoutCompositeSettings* Settings = GetDefault<UHoldoutCompositeSettings>();
-	if (Settings != nullptr)
-	{
-		return Settings->SceneViewExtensionPriority;
-	}
 
-	return 0;
+	return Settings ? Settings->CompositePassPriority : 0;
 }
 
 void FHoldoutCompositeSceneViewExtension::SetupViewFamily(FSceneViewFamily& InViewFamily)
@@ -246,8 +247,10 @@ void FHoldoutCompositeSceneViewExtension::SetupViewFamily(FSceneViewFamily& InVi
 	const UHoldoutCompositeSettings* Settings = GetDefault<UHoldoutCompositeSettings>();
 	if (Settings != nullptr)
 	{
-		bCompositeFollowsSceneExposure = Settings->bCompositeFollowsSceneExposure;
-		bCompositeSupportsSSR = Settings->bCompositeSupportsSSR;
+		bFollowSceneExposure = Settings->bFollowSceneExposure;
+		bSupportSSR	= Settings->bSupportSSR;
+		bApplyTSRLensDistortion = Settings->bApplyTSRLensDistortion;
+		CompositePassLocation = Settings->CompositePassLocation;
 	}
 }
 
@@ -311,7 +314,7 @@ void FHoldoutCompositeSceneViewExtension::SetupView(FSceneViewFamily& InViewFami
 	PassInput.EngineShowFlags = InViewFamily.EngineShowFlags;
 	PassInput.EngineShowFlags.DisableFeaturesForUnlit();
 	PassInput.EngineShowFlags.SetTranslucency(true);
-	PassInput.EngineShowFlags.SetUnlitViewmode(Settings ? Settings->bUnlitViewmode : false);
+	PassInput.EngineShowFlags.SetUnlitViewmode(Settings ? Settings->bSupportUnlitViewmode : false);
 	PassInput.EngineShowFlags.SetAllowPrimitiveAlphaHoldout(false);
 
 	// Note: Incoming view location is invalid for scene captures
@@ -336,14 +339,18 @@ void FHoldoutCompositeSceneViewExtension::SubscribeToPostProcessingPass(EPostPro
 		// Early-out to avoid needless work in the post processing callback(s).
 		return;
 	}
-
-	if (PassId == EPostProcessingPass::SSRInput && bCompositeSupportsSSR.load())
+	
+	if (PassId == EPostProcessingPass::SSRInput && bSupportSSR.load())
 	{
 		InOutPassCallbacks.Add(FAfterPassCallbackDelegate::CreateRaw(this, &FHoldoutCompositeSceneViewExtension::PostProcessPassSSRInput_RenderThread));
 	}
-	else if (PassId == EPostProcessingPass::Tonemap)
+	else if (PassId == EPostProcessingPass::Tonemap && CompositePassLocation.load() == EHoldoutCompositePassLocation::AfterTonemap)
+	{
+		InOutPassCallbacks.Add(FAfterPassCallbackDelegate::CreateRaw(this, &FHoldoutCompositeSceneViewExtension::PostProcessCompositePass_RenderThread));
+	}
+	else if (PassId == EPostProcessingPass::VisualizeDepthOfField && CompositePassLocation.load() == EHoldoutCompositePassLocation::EndOfPostProcessing)
 	{
-		InOutPassCallbacks.Add(FAfterPassCallbackDelegate::CreateRaw(this, &FHoldoutCompositeSceneViewExtension::PostProcessPassAfterTonemap_RenderThread));
+		InOutPassCallbacks.Add(FAfterPassCallbackDelegate::CreateRaw(this, &FHoldoutCompositeSceneViewExtension::PostProcessCompositePass_RenderThread));
 	}
 }
 
@@ -402,7 +409,7 @@ FScreenPassTexture FHoldoutCompositeSceneViewExtension::PostProcessPassSSRInput_
 	FHoldoutCompositeSSRInputShader::FParameters* PassParameters = GraphBuilder.AllocParameters<FHoldoutCompositeSSRInputShader::FParameters>();
 	constexpr bool bIsSceneColorUndistorted = true;
 	PassParameters->Common = BuildCommonCompositeParameters(GraphBuilder, InView, SceneColor, Output, bIsSceneColorUndistorted);
-	PassParameters->LastGlobalExposure = bCompositeFollowsSceneExposure.load() ? InView.GetLastEyeAdaptationExposure() : 1.0f;
+	PassParameters->LastGlobalExposure = bFollowSceneExposure.load() ? InView.GetLastEyeAdaptationExposure() : 1.0f;
 	PassParameters->RenderTargets[0] = Output.GetRenderTargetBinding();
 
 	TShaderMapRef<FHoldoutCompositeSSRInputShader> PixelShader(GlobalShaderMap);
@@ -419,7 +426,7 @@ FScreenPassTexture FHoldoutCompositeSceneViewExtension::PostProcessPassSSRInput_
 	return Output;
 }
 
-FScreenPassTexture FHoldoutCompositeSceneViewExtension::PostProcessPassAfterTonemap_RenderThread(FRDGBuilder& GraphBuilder, const FSceneView& InView, const FPostProcessMaterialInputs& Inputs)
+FScreenPassTexture FHoldoutCompositeSceneViewExtension::PostProcessCompositePass_RenderThread(FRDGBuilder& GraphBuilder, const FSceneView& InView, const FPostProcessMaterialInputs& Inputs)
 {
 	RDG_EVENT_SCOPE_STAT(GraphBuilder, FHoldoutCompositeFinal, "HoldoutComposite.Final");
 	RDG_GPU_STAT_SCOPE(GraphBuilder, FHoldoutCompositeFinal);
@@ -454,18 +461,19 @@ FScreenPassTexture FHoldoutCompositeSceneViewExtension::PostProcessPassAfterTone
 	// Compositing pass
 	{
 		FHoldoutCompositeFinalShader::FPermutationDomain PermutationVector;
-		PermutationVector.Set<FHoldoutCompositeFinalShader::FUseGlobalExposure>(bCompositeFollowsSceneExposure.load());
+		PermutationVector.Set<FHoldoutCompositeFinalShader::FUseGlobalExposure>(bFollowSceneExposure.load());
 
 		FHoldoutCompositeFinalShader::FParameters* PassParameters = GraphBuilder.AllocParameters<FHoldoutCompositeFinalShader::FParameters>();
 		PassParameters->View = InView.ViewUniformBuffer;
 		PassParameters->Common = BuildCommonCompositeParameters(GraphBuilder, InView, SceneColor, Output);
-		if (bCompositeFollowsSceneExposure.load())
+		if (bFollowSceneExposure.load())
 		{
 			FRDGBufferRef EyeAdaptationBuffer = GraphBuilder.RegisterExternalBuffer(InView.GetEyeAdaptationBuffer(), ERDGBufferFlags::MultiFrame);
 			PassParameters->EyeAdaptationBuffer = GraphBuilder.CreateSRV(EyeAdaptationBuffer);
 		}
 		PassParameters->Encodings = Encodings;
 		PassParameters->DisplayGamma = FVector2f(Family->RenderTarget->GetDisplayGamma(), 1.0f / Family->RenderTarget->GetDisplayGamma());
+		PassParameters->bApplyTSRLensDistortion = bApplyTSRLensDistortion;
 		PassParameters->RenderTargets[0] = Output.GetRenderTargetBinding();
 
 		TShaderMapRef<FHoldoutCompositeFinalShader> PixelShader(GlobalShaderMap, PermutationVector);
diff --git a/Engine/Plugins/Experimental/Compositing/HoldoutComposite/Source/HoldoutComposite/Private/HoldoutCompositeSceneViewExtension.h b/Engine/Plugins/Experimental/Compositing/HoldoutComposite/Source/HoldoutComposite/Private/HoldoutCompositeSceneViewExtension.h
--- a/Engine/Plugins/Experimental/Compositing/HoldoutComposite/Source/HoldoutComposite/Private/HoldoutCompositeSceneViewExtension.h
+++ b/Engine/Plugins/Experimental/Compositing/HoldoutComposite/Source/HoldoutComposite/Private/HoldoutCompositeSceneViewExtension.h
@@ -10,6 +10,7 @@
 class FViewInfo;
 struct FScreenPassRenderTarget;
 struct FScreenPassTexture;
+enum class EHoldoutCompositePassLocation;
 
 class FHoldoutCompositeSceneViewExtension : public FWorldSceneViewExtension
 {
@@ -51,7 +52,7 @@ private:
 	class FHoldoutCompositeCommonParameters BuildCommonCompositeParameters(FRDGBuilder& GraphBuilder, const FSceneView& InView, const FScreenPassTexture& SceneColor, const FScreenPassRenderTarget& Output, bool bIsSceneColorUndistorted = false);
 
 	FScreenPassTexture PostProcessPassSSRInput_RenderThread(FRDGBuilder& GraphBuilder, const FSceneView& InView, const FPostProcessMaterialInputs& Inputs);
-	FScreenPassTexture PostProcessPassAfterTonemap_RenderThread(FRDGBuilder& GraphBuilder, const FSceneView& InView, const FPostProcessMaterialInputs& Inputs);
+	FScreenPassTexture PostProcessCompositePass_RenderThread(FRDGBuilder& GraphBuilder, const FSceneView& InView, const FPostProcessMaterialInputs& Inputs);
 
 	// Collection of primitives to render as a custom render pass and composite after post-processing.
 	TSet<TSoftObjectPtr<UPrimitiveComponent>> CompositePrimitives;
@@ -59,10 +60,10 @@ private:
 	// Custom render pass render targets for each active view
 	TMap<uint32, TRefCountPtr<IPooledRenderTarget>> CustomRenderTargetPerView_RenderThread;
 
-	// Flag to enable global exposure on the composited render
-	std::atomic_bool bCompositeFollowsSceneExposure = false;
-
-	// Flag to enable composite into screen-space reflections
-	std::atomic_bool bCompositeSupportsSSR = false;
+	// Duplicate flags from plugin settings for thread-safety
+	std::atomic_bool bFollowSceneExposure;
+	std::atomic_bool bSupportSSR;
+	std::atomic_bool bApplyTSRLensDistortion;
+	std::atomic<EHoldoutCompositePassLocation> CompositePassLocation;
 };
 
diff --git a/Engine/Plugins/Experimental/Compositing/HoldoutComposite/Source/HoldoutComposite/Private/HoldoutCompositeSettings.cpp b/Engine/Plugins/Experimental/Compositing/HoldoutComposite/Source/HoldoutComposite/Private/HoldoutCompositeSettings.cpp
--- a/Engine/Plugins/Experimental/Compositing/HoldoutComposite/Source/HoldoutComposite/Private/HoldoutCompositeSettings.cpp
+++ b/Engine/Plugins/Experimental/Compositing/HoldoutComposite/Source/HoldoutComposite/Private/HoldoutCompositeSettings.cpp
@@ -5,10 +5,13 @@
 #include UE_INLINE_GENERATED_CPP_BY_NAME(HoldoutCompositeSettings)
 
 UHoldoutCompositeSettings::UHoldoutCompositeSettings()
-	: bCompositeFollowsSceneExposure(false)
-	, bCompositeSupportsSSR(true)
-	, bUnlitViewmode(true)
-	, SceneViewExtensionPriority(100 + 1 /* OPENCOLORIO_SCENE_VIEW_EXTENSION_PRIORITY + 1 */)
+	: bFollowSceneExposure(false)
+	, bSupportSSR(true)
+	, bSupportUnlitViewmode(true)
+	, bApplyTSRLensDistortion(true)
+	, DisabledPrimitiveClasses()
+	, CompositePassLocation(EHoldoutCompositePassLocation::AfterTonemap)
+	, CompositePassPriority(100 + 1 /* OPENCOLORIO_SCENE_VIEW_EXTENSION_PRIORITY + 1 */)
 { }
 
 FName UHoldoutCompositeSettings::GetCategoryName() const
diff --git a/Engine/Plugins/Experimental/Compositing/HoldoutComposite/Source/HoldoutComposite/Public/HoldoutCompositeSettings.h b/Engine/Plugins/Experimental/Compositing/HoldoutComposite/Source/HoldoutComposite/Public/HoldoutCompositeSettings.h
--- a/Engine/Plugins/Experimental/Compositing/HoldoutComposite/Source/HoldoutComposite/Public/HoldoutCompositeSettings.h
+++ b/Engine/Plugins/Experimental/Compositing/HoldoutComposite/Source/HoldoutComposite/Public/HoldoutCompositeSettings.h
@@ -9,6 +9,13 @@
 
 #include "HoldoutCompositeSettings.generated.h"
 
+UENUM()
+enum class EHoldoutCompositePassLocation
+{
+	AfterTonemap,
+	EndOfPostProcessing
+};
+
 /**
  * Settings for the HoldoutComposite module.
  */
@@ -31,22 +38,30 @@ public:
 public:
 	/** When enabled, the view global exposure is applied onto the separate render when composited. */
 	UPROPERTY(config, EditAnywhere, Category = General)
-	bool bCompositeFollowsSceneExposure;
+	bool bFollowSceneExposure;
 
 	/** When enabled, the separate composited render is also used to update screen-space reflections. */
 	UPROPERTY(config, EditAnywhere, Category = General)
-	bool bCompositeSupportsSSR;
+	bool bSupportSSR;
 
 	/** Enables the development shader debug feature that routes the Base Color output to Emissive for the separate render. Non-shipping PC build only. */
 	UPROPERTY(config, EditAnywhere, Category = General)
-	bool bUnlitViewmode;
+	bool bSupportUnlitViewmode;
 
-	/** Composite (scene view extension) pass priority, which defaults to before OpenColorIO. */
-	UPROPERTY(config, EditAnywhere, Category = General, AdvancedDisplay)
-	int32 SceneViewExtensionPriority;
+	/** When enabled, the TSR lens distortion (r.TSR.LensDistortion / Scene View Extension (Experimental) on lens files) is also applied on the separate composited render. */
+	UPROPERTY(config, EditAnywhere, Category = General, DisplayName="Apply TSR Lens Distortion")
+	bool bApplyTSRLensDistortion;
 
 	/** Primitive component classes that do not support the holdout composite.*/
 	UPROPERTY(config, EditAnywhere, Category = General)
 	TArray<FSoftClassPath> DisabledPrimitiveClasses;
+
+	/** Composite (scene view extension) pass location, which defaults to "After Tonemap". "End Of Post Processing" may be used to push the composite pass after additional post-process materials (such as the regular lens distortion) and OCIO. */
+	UPROPERTY(config, EditAnywhere, Category = General, AdvancedDisplay)
+	EHoldoutCompositePassLocation CompositePassLocation;
+
+	/** Composite (scene view extension) pass priority, which defaults to before OpenColorIO. */
+	UPROPERTY(config, EditAnywhere, Category = General, AdvancedDisplay)
+	int32 CompositePassPriority;
 };
 
