HoldoutComposite: Adds support for FXAA on the separate composited render.

This may be useful for flat colored motion graphics.

diff --git a/Engine/Plugins/Experimental/Compositing/HoldoutComposite/Shaders/Private/HoldoutCompositeFinal.usf b/Engine/Plugins/Experimental/Compositing/HoldoutComposite/Shaders/Private/HoldoutCompositeFinal.usf
--- a/Engine/Plugins/Experimental/Compositing/HoldoutComposite/Shaders/Private/HoldoutCompositeFinal.usf
+++ b/Engine/Plugins/Experimental/Compositing/HoldoutComposite/Shaders/Private/HoldoutCompositeFinal.usf
@@ -22,8 +22,8 @@ Texture2D<float2> UndistortingDisplacementTexture;
 SamplerState UndistortingDisplacementSampler;
 
 StructuredBuffer<float4> EyeAdaptationBuffer;
-uint2 Encodings;
-float2 DisplayGamma;
+uint3 Encodings;
+float3 DisplayGamma;
 uint bApplyTSRLensDistortion;
 
 float SafeRcp(float x)
@@ -72,6 +72,12 @@ void MainPS(
         SceneColor.rgb = sRGBToLinear(SceneColor.rgb);
     }
 
+    // Optionally linearize custom color
+    if(Encodings.z > 0)
+    {
+        CustomColor.rgb = pow(CustomColor.rgb, DisplayGamma.z);
+    }
+
     // The alpha comes from the scene color holdout mask, so we unpremultiply the custom render to get its opaque colors. (No-op on opaque geometry.)
     CustomColor = InvertAlpha(Unpremultiply(InvertAlpha(CustomColor)));
 
diff --git a/Engine/Plugins/Experimental/Compositing/HoldoutComposite/Shaders/Private/HoldoutCompositeGamma.usf b/Engine/Plugins/Experimental/Compositing/HoldoutComposite/Shaders/Private/HoldoutCompositeGamma.usf
new file mode 100644
--- /dev/null
+++ b/Engine/Plugins/Experimental/Compositing/HoldoutComposite/Shaders/Private/HoldoutCompositeGamma.usf
@@ -0,0 +1,23 @@
+// Copyright Epic Games, Inc. All Rights Reserved.
+
+#include "/Engine/Private/Common.ush"
+
+#ifndef TILE_SIZE
+    #error TILE_SIZE must be defined
+#endif
+
+RWTexture2D<float4> RWOutputTexture;
+uint2 Dimensions;
+float GammaPower;
+
+[numthreads(TILE_SIZE, TILE_SIZE, 1)]
+void MainCS(
+    uint2 DispatchThreadId : SV_DispatchThreadID)
+{
+    if (any(DispatchThreadId >= Dimensions))
+    {
+        return;
+    }
+
+    RWOutputTexture[DispatchThreadId].rgb = pow(RWOutputTexture[DispatchThreadId].rgb, GammaPower);
+}
diff --git a/Engine/Plugins/Experimental/Compositing/HoldoutComposite/Source/HoldoutComposite/Private/HoldoutCompositeSceneViewExtension.cpp b/Engine/Plugins/Experimental/Compositing/HoldoutComposite/Source/HoldoutComposite/Private/HoldoutCompositeSceneViewExtension.cpp
--- a/Engine/Plugins/Experimental/Compositing/HoldoutComposite/Source/HoldoutComposite/Private/HoldoutCompositeSceneViewExtension.cpp
+++ b/Engine/Plugins/Experimental/Compositing/HoldoutComposite/Source/HoldoutComposite/Private/HoldoutCompositeSceneViewExtension.cpp
@@ -13,6 +13,7 @@
 #include "HDRHelper.h"
 #include "PixelShaderUtils.h"
 #include "PostProcess/LensDistortion.h"
+#include "PostProcess/PostProcessAA.h"
 #include "PostProcess/PostProcessMaterialInputs.h"
 #include "Rendering/CustomRenderPass.h"
 #include "RenderGraphBuilder.h"
@@ -22,6 +23,7 @@
 #include "SystemTextures.h"
 
 DECLARE_GPU_STAT_NAMED(FHoldoutCompositeDilate, TEXT("HoldoutComposite.Dilate"));
+DECLARE_GPU_STAT_NAMED(FHoldoutCompositeGamma, TEXT("HoldoutComposite.Gamma"));
 DECLARE_GPU_STAT_NAMED(FHoldoutCompositeSSRInput, TEXT("HoldoutComposite.SSRInput"));
 DECLARE_GPU_STAT_NAMED(FHoldoutCompositeFinal, TEXT("HoldoutComposite.Final"));
 
@@ -33,6 +35,8 @@ namespace HoldoutComposite
 		Gamma = 1,
 		sRGB = 2,
 	};
+
+	static constexpr float DefaultGamma = 2.2f;
 }
 
 class FHoldoutCompositeDilateShader : public FGlobalShader
@@ -57,6 +61,28 @@ class FHoldoutCompositeDilateShader : public FGlobalShader
 };
 IMPLEMENT_GLOBAL_SHADER(FHoldoutCompositeDilateShader, "/Plugin/HoldoutComposite/Private/HoldoutCompositeDilate.usf", "MainCS", SF_Compute);
 
+class FHoldoutCompositeGammaShader : public FGlobalShader
+{
+	DECLARE_GLOBAL_SHADER(FHoldoutCompositeGammaShader);
+	SHADER_USE_PARAMETER_STRUCT(FHoldoutCompositeGammaShader, FGlobalShader);
+
+	BEGIN_SHADER_PARAMETER_STRUCT(FParameters, )
+		SHADER_PARAMETER_RDG_TEXTURE_UAV(RWTexture2D<float4>, RWOutputTexture)
+		SHADER_PARAMETER(FIntPoint, Dimensions)
+		SHADER_PARAMETER(float, GammaPower)
+	END_SHADER_PARAMETER_STRUCT()
+
+	static inline void ModifyCompilationEnvironment(const FGlobalShaderPermutationParameters& Parameters, FShaderCompilerEnvironment& OutEnvironment)
+	{
+		OutEnvironment.SetDefine(TEXT("TILE_SIZE"), ThreadGroupSize);
+
+		FGlobalShader::ModifyCompilationEnvironment(Parameters, OutEnvironment);
+	}
+
+	static const uint32 ThreadGroupSize = 16;
+};
+IMPLEMENT_GLOBAL_SHADER(FHoldoutCompositeGammaShader, "/Plugin/HoldoutComposite/Private/HoldoutCompositeGamma.usf", "MainCS", SF_Compute);
+
 BEGIN_SHADER_PARAMETER_STRUCT(FHoldoutCompositeCommonParameters, )
 	SHADER_PARAMETER_STRUCT(FScreenPassTextureViewportParameters, Input)
 	SHADER_PARAMETER_STRUCT(FScreenPassTextureViewportParameters, Custom)
@@ -96,24 +122,76 @@ class FHoldoutCompositeFinalShader : public FGlobalShader
 		SHADER_PARAMETER_STRUCT_REF(FViewUniformShaderParameters, View)
 		SHADER_PARAMETER_STRUCT_INCLUDE(FHoldoutCompositeCommonParameters, Common)
 		SHADER_PARAMETER_RDG_BUFFER_SRV(StructuredBuffer<float4>, EyeAdaptationBuffer)
-		SHADER_PARAMETER(FUint32Vector2, Encodings)
-		SHADER_PARAMETER(FVector2f, DisplayGamma)
+		SHADER_PARAMETER(FUint32Vector, Encodings)
+		SHADER_PARAMETER(FVector3f, DisplayGamma)
 		SHADER_PARAMETER(uint32, bApplyTSRLensDistortion)
 		RENDER_TARGET_BINDING_SLOTS()
 	END_SHADER_PARAMETER_STRUCT()
 };
 IMPLEMENT_GLOBAL_SHADER(FHoldoutCompositeFinalShader, "/Plugin/HoldoutComposite/Private/HoldoutCompositeFinal.usf", "MainPS", SF_Pixel);
 
+namespace HoldoutComposite
+{
+	void AddDilatePass(FRDGBuilder& GraphBuilder, FRDGTextureRef Input, FRDGTextureRef Output, ERHIFeatureLevel::Type FeatureLevel)
+	{
+		RDG_EVENT_SCOPE_STAT(GraphBuilder, FHoldoutCompositeDilate, "HoldoutComposite.Dilate");
+		RDG_GPU_STAT_SCOPE(GraphBuilder, FHoldoutCompositeDilate);
+
+		FGlobalShaderMap* GlobalShaderMap = GetGlobalShaderMap(FeatureLevel);
+		const FIntPoint TextureSize = Input->Desc.Extent;
+
+		FHoldoutCompositeDilateShader::FParameters* PassParameters = GraphBuilder.AllocParameters<FHoldoutCompositeDilateShader::FParameters>();
+		PassParameters->InputTexture = Input;
+		PassParameters->RWOutputTexture = GraphBuilder.CreateUAV(Output);
+		PassParameters->Dimensions = TextureSize;
+
+		TShaderMapRef<FHoldoutCompositeDilateShader> ComputeShader(GlobalShaderMap);
+		FComputeShaderUtils::AddPass(
+			GraphBuilder,
+			RDG_EVENT_NAME("HoldoutComposite.Dilate (%dx%d)", TextureSize.X, TextureSize.Y),
+			GSupportsEfficientAsyncCompute ? ERDGPassFlags::AsyncCompute : ERDGPassFlags::Compute,
+			ComputeShader,
+			PassParameters,
+			FComputeShaderUtils::GetGroupCount(TextureSize, FHoldoutCompositeDilateShader::ThreadGroupSize)
+		);
+	}
+
+	void AddGammaPass(FRDGBuilder& GraphBuilder, FRDGTextureRef InOutTexture, ERHIFeatureLevel::Type FeatureLevel)
+	{
+		RDG_EVENT_SCOPE_STAT(GraphBuilder, FHoldoutCompositeGamma, "HoldoutComposite.Gamma");
+		RDG_GPU_STAT_SCOPE(GraphBuilder, FHoldoutCompositeGamma);
+
+		FGlobalShaderMap* GlobalShaderMap = GetGlobalShaderMap(FeatureLevel);
+		const FIntPoint TextureSize = InOutTexture->Desc.Extent;
+
+		FHoldoutCompositeGammaShader::FParameters* PassParameters = GraphBuilder.AllocParameters<FHoldoutCompositeGammaShader::FParameters>();
+		PassParameters->RWOutputTexture = GraphBuilder.CreateUAV(InOutTexture);
+		PassParameters->Dimensions = TextureSize;
+		PassParameters->GammaPower = 1.0f / HoldoutComposite::DefaultGamma;
+
+		TShaderMapRef<FHoldoutCompositeGammaShader> ComputeShader(GlobalShaderMap);
+		FComputeShaderUtils::AddPass(
+			GraphBuilder,
+			RDG_EVENT_NAME("HoldoutComposite.Gamma (%dx%d)", TextureSize.X, TextureSize.Y),
+			GSupportsEfficientAsyncCompute ? ERDGPassFlags::AsyncCompute : ERDGPassFlags::Compute,
+			ComputeShader,
+			PassParameters,
+			FComputeShaderUtils::GetGroupCount(TextureSize, FHoldoutCompositeGammaShader::ThreadGroupSize)
+		);
+	}
+}
+
 class FHoldoutCompositeCustomRenderPass : public FCustomRenderPassBase
 {
 public:
 	IMPLEMENT_CUSTOM_RENDER_PASS(FHoldoutCompositeCustomRenderPass);
 
-	FHoldoutCompositeCustomRenderPass(const FIntPoint& InRenderTargetSize, FHoldoutCompositeSceneViewExtension* InParentExtension, FSceneView& InView)
+	FHoldoutCompositeCustomRenderPass(const FIntPoint& InRenderTargetSize, FHoldoutCompositeSceneViewExtension* InParentExtension, FSceneView& InView, bool bInGammaEncode)
 		: FCustomRenderPassBase(TEXT("HoldoutCompositeCustomRenderPass"), ERenderMode::DepthAndBasePass, ERenderOutput::SceneColorAndAlpha, InRenderTargetSize)
 		, ParentExtension(InParentExtension)
 		, ViewId(InView.GetViewKey())
 		, ViewFeatureLevel(InView.GetFeatureLevel())
+		, bGammaEncode(bInGammaEncode)
 	{
 		bSceneColorWithTranslucent = true;
 	}
@@ -127,48 +205,25 @@ public:
 
 	virtual void OnPostRender(FRDGBuilder& GraphBuilder) override
 	{
-		FRDGTextureRef DilatedCRP = CreateDilatedTexture(GraphBuilder, RenderTargetTexture, ViewFeatureLevel);
-
-		ParentExtension->CollectCustomRenderTarget(ViewId, GraphBuilder.ConvertToExternalTexture(DilatedCRP));
-	}
-
-private:
-	
-	static FRDGTextureRef CreateDilatedTexture(FRDGBuilder& GraphBuilder, FRDGTextureRef InRenderTarget, ERHIFeatureLevel::Type FeatureLevel)
-	{
-		RDG_EVENT_SCOPE_STAT(GraphBuilder, FHoldoutCompositeDilate, "HoldoutComposite.Dilate");
-		RDG_GPU_STAT_SCOPE(GraphBuilder, FHoldoutCompositeDilate);
-
-		FGlobalShaderMap* GlobalShaderMap = GetGlobalShaderMap(FeatureLevel);
+		const FRDGTextureDesc TextureDesc = FRDGTextureDesc::Create2D(RenderTargetTexture->Desc.Extent, PF_FloatRGBA, FClearValueBinding::Black, TexCreate_UAV | TexCreate_ShaderResource);
+		FRDGTextureRef Output = GraphBuilder.CreateTexture(TextureDesc, TEXT("HoldoutCompositeProcessedTexture"));
 
-		const FIntPoint TextureSize = FIntPoint(InRenderTarget->Desc.GetSize().X, InRenderTarget->Desc.GetSize().Y);
-		const FRDGTextureDesc TextureDesc = FRDGTextureDesc::Create2D(TextureSize, PF_FloatRGBA, FClearValueBinding::Black, TexCreate_UAV | TexCreate_ShaderResource);
-		FRDGTextureRef DilatedTexture = GraphBuilder.CreateTexture(TextureDesc, TEXT("HoldoutCompositeDilatedTexture"));
+		HoldoutComposite::AddDilatePass(GraphBuilder, RenderTargetTexture, Output, ViewFeatureLevel);
 
-		// Async compute dilation pass
+		if (bGammaEncode)
 		{
-			FHoldoutCompositeDilateShader::FParameters* PassParameters = GraphBuilder.AllocParameters<FHoldoutCompositeDilateShader::FParameters>();
-			PassParameters->InputTexture = InRenderTarget;
-			PassParameters->RWOutputTexture = GraphBuilder.CreateUAV(DilatedTexture);
-			PassParameters->Dimensions = TextureSize;
-
-			TShaderMapRef<FHoldoutCompositeDilateShader> ComputeShader(GlobalShaderMap);
-			FComputeShaderUtils::AddPass(
-				GraphBuilder,
-				RDG_EVENT_NAME("HoldoutComposite.Dilate (%dx%d)", TextureSize.X, TextureSize.Y),
-				GSupportsEfficientAsyncCompute ? ERDGPassFlags::AsyncCompute : ERDGPassFlags::Compute,
-				ComputeShader,
-				PassParameters,
-				FComputeShaderUtils::GetGroupCount(TextureSize, FHoldoutCompositeDilateShader::ThreadGroupSize)
-			);
+			HoldoutComposite::AddGammaPass(GraphBuilder, Output, ViewFeatureLevel);
 		}
 
-		return DilatedTexture;
+		ParentExtension->CollectCustomRenderTarget(ViewId, GraphBuilder.ConvertToExternalTexture(Output));
 	}
 
+private:
+
 	FHoldoutCompositeSceneViewExtension* ParentExtension;
-	uint32 ViewId;
-	ERHIFeatureLevel::Type ViewFeatureLevel;
+	const uint32 ViewId;
+	const ERHIFeatureLevel::Type ViewFeatureLevel;
+	const bool bGammaEncode;
 };
 
 FHoldoutCompositeSceneViewExtension::FHoldoutCompositeSceneViewExtension(const FAutoRegister& AutoReg, UWorld* InWorld)
@@ -250,6 +305,7 @@ void FHoldoutCompositeSceneViewExtension::SetupViewFamily(FSceneViewFamily& InVi
 		bFollowSceneExposure = Settings->bFollowSceneExposure;
 		bSupportSSR	= Settings->bSupportSSR;
 		bApplyTSRLensDistortion = Settings->bApplyTSRLensDistortion;
+		bApplyFXAA = Settings->bApplyFXAA;
 		CompositePassLocation = Settings->CompositePassLocation;
 	}
 }
@@ -300,12 +356,16 @@ void FHoldoutCompositeSceneViewExtension::SetupView(FSceneViewFamily& InViewFami
 	{
 		RenderTargetViewSize = InView.UnscaledViewRect.Size();
 	}
+	
+	// If FXAA is active, we gamma-encode the result as FXAA should operate on non-linear colors
+	const bool bGammaEncode = bApplyFXAA;
 
 	// Create a new custom render pass to render the composite primitive(s)
 	FHoldoutCompositeCustomRenderPass* CustomRenderPass = new FHoldoutCompositeCustomRenderPass(
 		RenderTargetViewSize,
 		this,
-		InView
+		InView,
+		bGammaEncode
 	);
 
 	const UHoldoutCompositeSettings* Settings = GetDefault<UHoldoutCompositeSettings>();
@@ -367,20 +427,18 @@ FRDGTextureRef FHoldoutCompositeSceneViewExtension::GetCustomRenderPassTexture(F
 	return CustomRenderPassTexture;
 }
 
-FHoldoutCompositeCommonParameters FHoldoutCompositeSceneViewExtension::BuildCommonCompositeParameters(FRDGBuilder& GraphBuilder, const FSceneView& InView, const FScreenPassTexture& SceneColor, const FScreenPassRenderTarget& Output, bool bIsSceneColorUndistorted)
+FHoldoutCompositeCommonParameters FHoldoutCompositeSceneViewExtension::BuildCommonCompositeParameters(FRDGBuilder& GraphBuilder, const FSceneView& InView, const FScreenPassTexture& SceneColor, const FScreenPassTexture& CustomRenderPassTexture, const FScreenPassRenderTarget& Output, bool bIsSceneColorUndistorted)
 {
 	const FLensDistortionLUT& LensDistortionLUT = LensDistortion::GetLUTUnsafe(InView);
 	const bool bLensDistortionInTSR = (LensDistortion::GetPassLocationUnsafe(InView) == LensDistortion::EPassLocation::TSR);
 
-	FRDGTextureRef CustomRenderPassTexture = GetCustomRenderPassTexture(GraphBuilder, InView);
-
 	FHoldoutCompositeCommonParameters CommonParameters;
 	CommonParameters.Input = GetScreenPassTextureViewportParameters(FScreenPassTextureViewport(SceneColor));
 	CommonParameters.Custom = GetScreenPassTextureViewportParameters(FScreenPassTextureViewport(CustomRenderPassTexture));
 	CommonParameters.Output = GetScreenPassTextureViewportParameters(FScreenPassTextureViewport(Output));
 	CommonParameters.InputTexture = SceneColor.Texture;
 	CommonParameters.InputSampler = bIsSceneColorUndistorted ? TStaticSamplerState<SF_Bilinear, AM_Mirror, AM_Mirror>::GetRHI() : TStaticSamplerState<SF_Point>::GetRHI();
-	CommonParameters.CustomTexture = CustomRenderPassTexture;
+	CommonParameters.CustomTexture = CustomRenderPassTexture.Texture;
 	CommonParameters.CustomSampler = TStaticSamplerState<SF_Bilinear>::GetRHI();
 	CommonParameters.UndistortingDisplacementTexture = GSystemTextures.GetBlackDummy(GraphBuilder);
 	CommonParameters.UndistortingDisplacementSampler = TStaticSamplerState<SF_Bilinear>::GetRHI();
@@ -404,11 +462,12 @@ FScreenPassTexture FHoldoutCompositeSceneViewExtension::PostProcessPassSSRInput_
 
 	FScreenPassTexture SceneColor = FScreenPassTexture::CopyFromSlice(GraphBuilder, Inputs.GetInput(EPostProcessMaterialInput::SceneColor));
 	check(SceneColor.IsValid());
+	FRDGTextureRef CustomRenderPassTexture = GetCustomRenderPassTexture(GraphBuilder, InView);
 	FScreenPassRenderTarget Output = FScreenPassRenderTarget::CreateFromInput(GraphBuilder, SceneColor, InView.GetOverwriteLoadAction(), TEXT("HoldoutCompositeSSRInputRT"));
 
 	FHoldoutCompositeSSRInputShader::FParameters* PassParameters = GraphBuilder.AllocParameters<FHoldoutCompositeSSRInputShader::FParameters>();
 	constexpr bool bIsSceneColorUndistorted = true;
-	PassParameters->Common = BuildCommonCompositeParameters(GraphBuilder, InView, SceneColor, Output, bIsSceneColorUndistorted);
+	PassParameters->Common = BuildCommonCompositeParameters(GraphBuilder, InView, SceneColor, FScreenPassTexture(CustomRenderPassTexture), Output, bIsSceneColorUndistorted);
 	PassParameters->LastGlobalExposure = bFollowSceneExposure.load() ? InView.GetLastEyeAdaptationExposure() : 1.0f;
 	PassParameters->RenderTargets[0] = Output.GetRenderTargetBinding();
 
@@ -435,12 +494,22 @@ FScreenPassTexture FHoldoutCompositeSceneViewExtension::PostProcessCompositePass
 	FScreenPassTexture SceneColor = FScreenPassTexture::CopyFromSlice(GraphBuilder, Inputs.GetInput(EPostProcessMaterialInput::SceneColor));
 	check(SceneColor.IsValid());
 
-	FRDGTextureRef CustomRenderPassTexture = GetCustomRenderPassTexture(GraphBuilder, InView);
+	FScreenPassTexture CustomRenderPassTexture = FScreenPassTexture(GetCustomRenderPassTexture(GraphBuilder, InView));
+
+	if (bApplyFXAA)
+	{
+		static TConsoleVariableData<int32>* CVarFXAAQuality = IConsoleManager::Get().FindTConsoleVariableDataInt(TEXT("r.FXAA.Quality"));
+
+		FFXAAInputs PassInputs{};
+		PassInputs.SceneColor = CustomRenderPassTexture;
+		PassInputs.Quality = static_cast<EFXAAQuality>(FMath::Clamp(CVarFXAAQuality->GetValueOnRenderThread(), 0, 5));
+		CustomRenderPassTexture = AddFXAAPass(GraphBuilder, InView, PassInputs);
+	}
 
 	FGlobalShaderMap* GlobalShaderMap = GetGlobalShaderMap(InView.GetFeatureLevel());
 	const FSceneViewFamily* Family = InView.Family;
 
-	FUint32Vector2 Encodings = FUint32Vector2::ZeroValue; // ESceneColorSourceEncoding::Linear
+	FUint32Vector Encodings = FUint32Vector::ZeroValue; // ESceneColorSourceEncoding::Linear
 	if ((Family->EngineShowFlags.Tonemapper == 0) || (Family->EngineShowFlags.PostProcessing == 0))
 	{
 		Encodings.X = static_cast<uint32>(ESceneColorSourceEncoding::Gamma);
@@ -451,6 +520,8 @@ FScreenPassTexture FHoldoutCompositeSceneViewExtension::PostProcessCompositePass
 		Encodings.X = static_cast<uint32>(ESceneColorSourceEncoding::sRGB);
 		Encodings.Y = static_cast<uint32>(ESceneColorSourceEncoding::sRGB);
 	}
+	// FXAA received gamma-encoded colors, so we must linearize before compositing.
+	Encodings.Z = static_cast<uint32>(bApplyFXAA);
 
 	FScreenPassRenderTarget Output = Inputs.OverrideOutput;
 	if (!Output.IsValid())
@@ -465,14 +536,14 @@ FScreenPassTexture FHoldoutCompositeSceneViewExtension::PostProcessCompositePass
 
 		FHoldoutCompositeFinalShader::FParameters* PassParameters = GraphBuilder.AllocParameters<FHoldoutCompositeFinalShader::FParameters>();
 		PassParameters->View = InView.ViewUniformBuffer;
-		PassParameters->Common = BuildCommonCompositeParameters(GraphBuilder, InView, SceneColor, Output);
+		PassParameters->Common = BuildCommonCompositeParameters(GraphBuilder, InView, SceneColor, CustomRenderPassTexture, Output);
 		if (bFollowSceneExposure.load())
 		{
 			FRDGBufferRef EyeAdaptationBuffer = GraphBuilder.RegisterExternalBuffer(InView.GetEyeAdaptationBuffer(), ERDGBufferFlags::MultiFrame);
 			PassParameters->EyeAdaptationBuffer = GraphBuilder.CreateSRV(EyeAdaptationBuffer);
 		}
 		PassParameters->Encodings = Encodings;
-		PassParameters->DisplayGamma = FVector2f(Family->RenderTarget->GetDisplayGamma(), 1.0f / Family->RenderTarget->GetDisplayGamma());
+		PassParameters->DisplayGamma = FVector3f(Family->RenderTarget->GetDisplayGamma(), 1.0f / Family->RenderTarget->GetDisplayGamma(), HoldoutComposite::DefaultGamma);
 		PassParameters->bApplyTSRLensDistortion = bApplyTSRLensDistortion;
 		PassParameters->RenderTargets[0] = Output.GetRenderTargetBinding();
 
diff --git a/Engine/Plugins/Experimental/Compositing/HoldoutComposite/Source/HoldoutComposite/Private/HoldoutCompositeSceneViewExtension.h b/Engine/Plugins/Experimental/Compositing/HoldoutComposite/Source/HoldoutComposite/Private/HoldoutCompositeSceneViewExtension.h
--- a/Engine/Plugins/Experimental/Compositing/HoldoutComposite/Source/HoldoutComposite/Private/HoldoutCompositeSceneViewExtension.h
+++ b/Engine/Plugins/Experimental/Compositing/HoldoutComposite/Source/HoldoutComposite/Private/HoldoutCompositeSceneViewExtension.h
@@ -49,7 +49,7 @@ protected:
 private:
 
 	FRDGTextureRef GetCustomRenderPassTexture(FRDGBuilder& GraphBuilder, const FSceneView& InView) const;
-	class FHoldoutCompositeCommonParameters BuildCommonCompositeParameters(FRDGBuilder& GraphBuilder, const FSceneView& InView, const FScreenPassTexture& SceneColor, const FScreenPassRenderTarget& Output, bool bIsSceneColorUndistorted = false);
+	class FHoldoutCompositeCommonParameters BuildCommonCompositeParameters(FRDGBuilder& GraphBuilder, const FSceneView& InView, const FScreenPassTexture& SceneColor, const FScreenPassTexture& CustomRenderPassTexture, const FScreenPassRenderTarget& Output, bool bIsSceneColorUndistorted = false);
 
 	FScreenPassTexture PostProcessPassSSRInput_RenderThread(FRDGBuilder& GraphBuilder, const FSceneView& InView, const FPostProcessMaterialInputs& Inputs);
 	FScreenPassTexture PostProcessCompositePass_RenderThread(FRDGBuilder& GraphBuilder, const FSceneView& InView, const FPostProcessMaterialInputs& Inputs);
@@ -64,6 +64,7 @@ private:
 	std::atomic_bool bFollowSceneExposure;
 	std::atomic_bool bSupportSSR;
 	std::atomic_bool bApplyTSRLensDistortion;
+	std::atomic_bool bApplyFXAA;
 	std::atomic<EHoldoutCompositePassLocation> CompositePassLocation;
 };
 
diff --git a/Engine/Plugins/Experimental/Compositing/HoldoutComposite/Source/HoldoutComposite/Private/HoldoutCompositeSettings.cpp b/Engine/Plugins/Experimental/Compositing/HoldoutComposite/Source/HoldoutComposite/Private/HoldoutCompositeSettings.cpp
--- a/Engine/Plugins/Experimental/Compositing/HoldoutComposite/Source/HoldoutComposite/Private/HoldoutCompositeSettings.cpp
+++ b/Engine/Plugins/Experimental/Compositing/HoldoutComposite/Source/HoldoutComposite/Private/HoldoutCompositeSettings.cpp
@@ -9,6 +9,7 @@ UHoldoutCompositeSettings::UHoldoutCompositeSettings()
 	, bSupportSSR(true)
 	, bSupportUnlitViewmode(true)
 	, bApplyTSRLensDistortion(true)
+	, bApplyFXAA(false)
 	, DisabledPrimitiveClasses()
 	, CompositePassLocation(EHoldoutCompositePassLocation::AfterTonemap)
 	, CompositePassPriority(100 + 1 /* OPENCOLORIO_SCENE_VIEW_EXTENSION_PRIORITY + 1 */)
diff --git a/Engine/Plugins/Experimental/Compositing/HoldoutComposite/Source/HoldoutComposite/Public/HoldoutCompositeSettings.h b/Engine/Plugins/Experimental/Compositing/HoldoutComposite/Source/HoldoutComposite/Public/HoldoutCompositeSettings.h
--- a/Engine/Plugins/Experimental/Compositing/HoldoutComposite/Source/HoldoutComposite/Public/HoldoutCompositeSettings.h
+++ b/Engine/Plugins/Experimental/Compositing/HoldoutComposite/Source/HoldoutComposite/Public/HoldoutCompositeSettings.h
@@ -52,6 +52,10 @@ public:
 	UPROPERTY(config, EditAnywhere, Category = General, DisplayName="Apply TSR Lens Distortion")
 	bool bApplyTSRLensDistortion;
 
+	/** When enabled, FXAA is applied on the separate composited render. */
+	UPROPERTY(config, EditAnywhere, Category = General, DisplayName = "Apply FXAA")
+	bool bApplyFXAA;
+
 	/** Primitive component classes that do not support the holdout composite.*/
 	UPROPERTY(config, EditAnywhere, Category = General)
 	TArray<FSoftClassPath> DisabledPrimitiveClasses;
