HoldoutComposite: Fix previous CL such that gamma encoding doesn't impact other passes (SSR input).

#rnx

diff --git a/Engine/Plugins/Experimental/Compositing/HoldoutComposite/Source/HoldoutComposite/Private/HoldoutCompositeSceneViewExtension.cpp b/Engine/Plugins/Experimental/Compositing/HoldoutComposite/Source/HoldoutComposite/Private/HoldoutCompositeSceneViewExtension.cpp
--- a/Engine/Plugins/Experimental/Compositing/HoldoutComposite/Source/HoldoutComposite/Private/HoldoutCompositeSceneViewExtension.cpp
+++ b/Engine/Plugins/Experimental/Compositing/HoldoutComposite/Source/HoldoutComposite/Private/HoldoutCompositeSceneViewExtension.cpp
@@ -156,7 +156,7 @@ namespace HoldoutComposite
 		);
 	}
 
-	void AddGammaPass(FRDGBuilder& GraphBuilder, FRDGTextureRef InOutTexture, ERHIFeatureLevel::Type FeatureLevel)
+	void AddGammaPass(FRDGBuilder& GraphBuilder, FRDGTextureRef InOutTexture, ERHIFeatureLevel::Type FeatureLevel, float GammaPower)
 	{
 		RDG_EVENT_SCOPE_STAT(GraphBuilder, FHoldoutCompositeGamma, "HoldoutComposite.Gamma");
 		RDG_GPU_STAT_SCOPE(GraphBuilder, FHoldoutCompositeGamma);
@@ -167,7 +167,7 @@ namespace HoldoutComposite
 		FHoldoutCompositeGammaShader::FParameters* PassParameters = GraphBuilder.AllocParameters<FHoldoutCompositeGammaShader::FParameters>();
 		PassParameters->RWOutputTexture = GraphBuilder.CreateUAV(InOutTexture);
 		PassParameters->Dimensions = TextureSize;
-		PassParameters->GammaPower = 1.0f / HoldoutComposite::DefaultGamma;
+		PassParameters->GammaPower = GammaPower;
 
 		TShaderMapRef<FHoldoutCompositeGammaShader> ComputeShader(GlobalShaderMap);
 		FComputeShaderUtils::AddPass(
@@ -186,12 +186,11 @@ class FHoldoutCompositeCustomRenderPass : public FCustomRenderPassBase
 public:
 	IMPLEMENT_CUSTOM_RENDER_PASS(FHoldoutCompositeCustomRenderPass);
 
-	FHoldoutCompositeCustomRenderPass(const FIntPoint& InRenderTargetSize, FHoldoutCompositeSceneViewExtension* InParentExtension, FSceneView& InView, bool bInGammaEncode)
+	FHoldoutCompositeCustomRenderPass(const FIntPoint& InRenderTargetSize, FHoldoutCompositeSceneViewExtension* InParentExtension, FSceneView& InView)
 		: FCustomRenderPassBase(TEXT("HoldoutCompositeCustomRenderPass"), ERenderMode::DepthAndBasePass, ERenderOutput::SceneColorAndAlpha, InRenderTargetSize)
 		, ParentExtension(InParentExtension)
 		, ViewId(InView.GetViewKey())
 		, ViewFeatureLevel(InView.GetFeatureLevel())
-		, bGammaEncode(bInGammaEncode)
 	{
 		bSceneColorWithTranslucent = true;
 	}
@@ -210,11 +209,6 @@ public:
 
 		HoldoutComposite::AddDilatePass(GraphBuilder, RenderTargetTexture, Output, ViewFeatureLevel);
 
-		if (bGammaEncode)
-		{
-			HoldoutComposite::AddGammaPass(GraphBuilder, Output, ViewFeatureLevel);
-		}
-
 		ParentExtension->CollectCustomRenderTarget(ViewId, GraphBuilder.ConvertToExternalTexture(Output));
 	}
 
@@ -223,7 +217,6 @@ private:
 	FHoldoutCompositeSceneViewExtension* ParentExtension;
 	const uint32 ViewId;
 	const ERHIFeatureLevel::Type ViewFeatureLevel;
-	const bool bGammaEncode;
 };
 
 FHoldoutCompositeSceneViewExtension::FHoldoutCompositeSceneViewExtension(const FAutoRegister& AutoReg, UWorld* InWorld)
@@ -356,16 +349,12 @@ void FHoldoutCompositeSceneViewExtension::SetupView(FSceneViewFamily& InViewFami
 	{
 		RenderTargetViewSize = InView.UnscaledViewRect.Size();
 	}
-	
-	// If FXAA is active, we gamma-encode the result as FXAA should operate on non-linear colors
-	const bool bGammaEncode = bApplyFXAA;
 
 	// Create a new custom render pass to render the composite primitive(s)
 	FHoldoutCompositeCustomRenderPass* CustomRenderPass = new FHoldoutCompositeCustomRenderPass(
 		RenderTargetViewSize,
 		this,
-		InView,
-		bGammaEncode
+		InView
 	);
 
 	const UHoldoutCompositeSettings* Settings = GetDefault<UHoldoutCompositeSettings>();
@@ -498,6 +487,9 @@ FScreenPassTexture FHoldoutCompositeSceneViewExtension::PostProcessCompositePass
 
 	if (bApplyFXAA)
 	{
+		// We gamma-encode the result so that FXAA can operate on non-linear colors
+		HoldoutComposite::AddGammaPass(GraphBuilder, CustomRenderPassTexture.Texture, InView.GetFeatureLevel(), 1.0f / HoldoutComposite::DefaultGamma);
+
 		static TConsoleVariableData<int32>* CVarFXAAQuality = IConsoleManager::Get().FindTConsoleVariableDataInt(TEXT("r.FXAA.Quality"));
 
 		FFXAAInputs PassInputs{};
