HoldoutComposite: Fix for horde shader cook errors.

diff --git a/Engine/Plugins/Experimental/Compositing/HoldoutComposite/Shaders/Private/HoldoutCompositeGamma.usf b/Engine/Plugins/Experimental/Compositing/HoldoutComposite/Shaders/Private/HoldoutCompositeGamma.usf
--- a/Engine/Plugins/Experimental/Compositing/HoldoutComposite/Shaders/Private/HoldoutCompositeGamma.usf
+++ b/Engine/Plugins/Experimental/Compositing/HoldoutComposite/Shaders/Private/HoldoutCompositeGamma.usf
@@ -19,5 +19,8 @@ void MainCS(
         return;
     }
 
-    RWOutputTexture[DispatchThreadId].rgb = pow(RWOutputTexture[DispatchThreadId].rgb, GammaPower);
+    float4 TexelColor = RWOutputTexture[DispatchThreadId];
+    TexelColor.rgb = pow(TexelColor.rgb, GammaPower);
+
+    RWOutputTexture[DispatchThreadId].rgb = TexelColor;
 }
