HoldoutComposite: Corrected fix for horde cook issue.

#rnx

diff --git a/Engine/Plugins/Experimental/Compositing/HoldoutComposite/Shaders/Private/HoldoutCompositeGamma.usf b/Engine/Plugins/Experimental/Compositing/HoldoutComposite/Shaders/Private/HoldoutCompositeGamma.usf
--- a/Engine/Plugins/Experimental/Compositing/HoldoutComposite/Shaders/Private/HoldoutCompositeGamma.usf
+++ b/Engine/Plugins/Experimental/Compositing/HoldoutComposite/Shaders/Private/HoldoutCompositeGamma.usf
@@ -22,5 +22,5 @@ void MainCS(
     float4 TexelColor = RWOutputTexture[DispatchThreadId];
     TexelColor.rgb = pow(TexelColor.rgb, GammaPower);
 
-    RWOutputTexture[DispatchThreadId].rgb = TexelColor;
+    RWOutputTexture[DispatchThreadId] = TexelColor;
 }
