HoldoutComposite: Additional required fixes for compute shader. (Apologies!)

#rnx

diff --git a/Engine/Plugins/Experimental/Compositing/HoldoutComposite/Shaders/Private/HoldoutCompositeGamma.usf b/Engine/Plugins/Experimental/Compositing/HoldoutComposite/Shaders/Private/HoldoutCompositeGamma.usf
--- a/Engine/Plugins/Experimental/Compositing/HoldoutComposite/Shaders/Private/HoldoutCompositeGamma.usf
+++ b/Engine/Plugins/Experimental/Compositing/HoldoutComposite/Shaders/Private/HoldoutCompositeGamma.usf
@@ -6,6 +6,7 @@
     #error TILE_SIZE must be defined
 #endif
 
+Texture2D InputTexture;
 RWTexture2D<float4> RWOutputTexture;
 uint2 Dimensions;
 float GammaPower;
@@ -19,8 +20,8 @@ void MainCS(
         return;
     }
 
-    float4 TexelColor = RWOutputTexture[DispatchThreadId];
-    TexelColor.rgb = pow(TexelColor.rgb, GammaPower);
+    float4 TexelColor = InputTexture[DispatchThreadId];
+    TexelColor.rgb = pow(TexelColor.rgb, abs(GammaPower));
 
     RWOutputTexture[DispatchThreadId] = TexelColor;
 }
diff --git a/Engine/Plugins/Experimental/Compositing/HoldoutComposite/Source/HoldoutComposite/Private/HoldoutCompositeSceneViewExtension.cpp b/Engine/Plugins/Experimental/Compositing/HoldoutComposite/Source/HoldoutComposite/Private/HoldoutCompositeSceneViewExtension.cpp
--- a/Engine/Plugins/Experimental/Compositing/HoldoutComposite/Source/HoldoutComposite/Private/HoldoutCompositeSceneViewExtension.cpp
+++ b/Engine/Plugins/Experimental/Compositing/HoldoutComposite/Source/HoldoutComposite/Private/HoldoutCompositeSceneViewExtension.cpp
@@ -67,6 +67,7 @@ class FHoldoutCompositeGammaShader : public FGlobalShader
 	SHADER_USE_PARAMETER_STRUCT(FHoldoutCompositeGammaShader, FGlobalShader);
 
 	BEGIN_SHADER_PARAMETER_STRUCT(FParameters, )
+		SHADER_PARAMETER_RDG_TEXTURE(Texture2D, InputTexture)
 		SHADER_PARAMETER_RDG_TEXTURE_UAV(RWTexture2D<float4>, RWOutputTexture)
 		SHADER_PARAMETER(FIntPoint, Dimensions)
 		SHADER_PARAMETER(float, GammaPower)
@@ -165,6 +166,7 @@ namespace HoldoutComposite
 		const FIntPoint TextureSize = InOutTexture->Desc.Extent;
 
 		FHoldoutCompositeGammaShader::FParameters* PassParameters = GraphBuilder.AllocParameters<FHoldoutCompositeGammaShader::FParameters>();
+		PassParameters->InputTexture = InOutTexture;
 		PassParameters->RWOutputTexture = GraphBuilder.CreateUAV(InOutTexture);
 		PassParameters->Dimensions = TextureSize;
 		PassParameters->GammaPower = GammaPower;
