HoldoutComposite: Allow additional PPM passes on separate composited render.

Immediately useful for shadow/reflection catching prototypes. Interface is likely to change as we improve integration of scene captures & user textures. Unified interface back to regular pointers, but the SVE stores weak object pointers internally.

diff --git a/Engine/Plugins/Experimental/Compositing/HoldoutComposite/Source/HoldoutComposite/Private/HoldoutCompositeComponent.cpp b/Engine/Plugins/Experimental/Compositing/HoldoutComposite/Source/HoldoutComposite/Private/HoldoutCompositeComponent.cpp
--- a/Engine/Plugins/Experimental/Compositing/HoldoutComposite/Source/HoldoutComposite/Private/HoldoutCompositeComponent.cpp
+++ b/Engine/Plugins/Experimental/Compositing/HoldoutComposite/Source/HoldoutComposite/Private/HoldoutCompositeComponent.cpp
@@ -37,9 +37,9 @@ namespace
 	}
 
 	// Find list of primitives from the parent component and its children. We need to traverse children to support objects such as UText3DComponent.
-	static TArray<TSoftObjectPtr<UPrimitiveComponent>> FindPrimitiveComponents(USceneComponent* InParentComponent)
+	static TArray<UPrimitiveComponent*> FindPrimitiveComponents(USceneComponent* InParentComponent)
 	{
-		TArray<TSoftObjectPtr<UPrimitiveComponent>> OutPrimitiveComponents;
+		TArray<UPrimitiveComponent*> OutPrimitiveComponents;
 
 		if (IsValid(InParentComponent))
 		{
@@ -106,7 +106,7 @@ void UHoldoutCompositeComponent::OnAttachmentChanged()
 	USceneComponent* SceneComponent = GetAttachParent();
 	if (IsValid(SceneComponent))
 	{
-		TArray<TSoftObjectPtr<UPrimitiveComponent>> ParentPrimitives = FindPrimitiveComponents(SceneComponent);
+		TArray<UPrimitiveComponent*> ParentPrimitives = FindPrimitiveComponents(SceneComponent);
 		if (ParentPrimitives.IsEmpty())
 		{
 #if WITH_EDITOR
@@ -154,7 +154,7 @@ void UHoldoutCompositeComponent::RegisterCompositeImpl()
 	}
 
 	UHoldoutCompositeSubsystem* Subsystem = UWorld::GetSubsystem<UHoldoutCompositeSubsystem>(GetWorld());
-	TArray<TSoftObjectPtr<UPrimitiveComponent>> ParentPrimitives = FindPrimitiveComponents(GetAttachParent());
+	TArray<UPrimitiveComponent*> ParentPrimitives = FindPrimitiveComponents(GetAttachParent());
 
 	if (IsValid(Subsystem) && !ParentPrimitives.IsEmpty())
 	{
@@ -165,7 +165,7 @@ void UHoldoutCompositeComponent::RegisterCompositeImpl()
 void UHoldoutCompositeComponent::UnregisterCompositeImpl()
 {
 	UHoldoutCompositeSubsystem* Subsystem = UWorld::GetSubsystem<UHoldoutCompositeSubsystem>(GetWorld());
-	TArray<TSoftObjectPtr<UPrimitiveComponent>> ParentPrimitives = FindPrimitiveComponents(GetAttachParent());
+	TArray<UPrimitiveComponent*> ParentPrimitives = FindPrimitiveComponents(GetAttachParent());
 
 	if (IsValid(Subsystem) && !ParentPrimitives.IsEmpty())
 	{
diff --git a/Engine/Plugins/Experimental/Compositing/HoldoutComposite/Source/HoldoutComposite/Private/HoldoutCompositeSceneViewExtension.cpp b/Engine/Plugins/Experimental/Compositing/HoldoutComposite/Source/HoldoutComposite/Private/HoldoutCompositeSceneViewExtension.cpp
--- a/Engine/Plugins/Experimental/Compositing/HoldoutComposite/Source/HoldoutComposite/Private/HoldoutCompositeSceneViewExtension.cpp
+++ b/Engine/Plugins/Experimental/Compositing/HoldoutComposite/Source/HoldoutComposite/Private/HoldoutCompositeSceneViewExtension.cpp
@@ -11,6 +11,7 @@
 #include "EngineUtils.h"
 #include "FXRenderingUtils.h"
 #include "HDRHelper.h"
+#include "Materials/Material.h"
 #include "PixelShaderUtils.h"
 #include "PostProcess/LensDistortion.h"
 #include "PostProcess/PostProcessAA.h"
@@ -234,13 +235,13 @@ FHoldoutCompositeSceneViewExtension::~FHoldoutCompositeSceneViewExtension()
 {
 }
 
-void FHoldoutCompositeSceneViewExtension::RegisterPrimitives(TArrayView<TSoftObjectPtr<UPrimitiveComponent>> InPrimitiveComponents, bool bInHoldoutState)
+void FHoldoutCompositeSceneViewExtension::RegisterPrimitives_GameThread(TArrayView<UPrimitiveComponent*> InPrimitiveComponents, bool bInHoldoutState)
 {
 	check(IsInGameThread());
 
-	for (const TSoftObjectPtr<UPrimitiveComponent>& InPrimitiveComponent : InPrimitiveComponents)
+	for (UPrimitiveComponent* InPrimitiveComponent : InPrimitiveComponents)
 	{
-		if (!InPrimitiveComponent.IsValid())
+		if (!IsValid(InPrimitiveComponent))
 		{
 			continue;
 		}
@@ -253,13 +254,13 @@ void FHoldoutCompositeSceneViewExtension::RegisterPrimitives(TArrayView<TSoftObj
 	}
 }
 
-void FHoldoutCompositeSceneViewExtension::UnregisterPrimitives(TArrayView<TSoftObjectPtr<UPrimitiveComponent>> InPrimitiveComponents, bool bInHoldoutState)
+void FHoldoutCompositeSceneViewExtension::UnregisterPrimitives_GameThread(TArrayView<UPrimitiveComponent*> InPrimitiveComponents, bool bInHoldoutState)
 {
 	check(IsInGameThread());
 
-	for (const TSoftObjectPtr<UPrimitiveComponent>& InPrimitiveComponent : InPrimitiveComponents)
+	for (UPrimitiveComponent* InPrimitiveComponent : InPrimitiveComponents)
 	{
-		if (!InPrimitiveComponent.IsValid())
+		if (!IsValid(InPrimitiveComponent))
 		{
 			continue;
 		}
@@ -270,7 +271,22 @@ void FHoldoutCompositeSceneViewExtension::UnregisterPrimitives(TArrayView<TSoftO
 			InPrimitiveComponent->SetHoldout(bInHoldoutState);
 		}
 	}
+}
+
+void FHoldoutCompositeSceneViewExtension::SetCompositePostProcessMaterials_GameThread(TArrayView<UMaterial*> InMaterials)
+{
+	check(IsInGameThread());
+
+	FScopeLock Lock(&MaterialsCS);
+	CompositeProcessMaterials = InMaterials;
+}
+
+void FHoldoutCompositeSceneViewExtension::ResetCompositePostProcessMaterials_GameThread()
+{
+	check(IsInGameThread());
 
+	FScopeLock Lock(&MaterialsCS);
+	CompositeProcessMaterials.Reset();
 }
 
 bool FHoldoutCompositeSceneViewExtension::IsActiveThisFrame_Internal(const FSceneViewExtensionContext& Context) const
@@ -311,22 +327,23 @@ void FHoldoutCompositeSceneViewExtension::SetupView(FSceneViewFamily& InViewFami
 	check(WorldPtr.IsValid());
 
 	TSet<FPrimitiveComponentId> HoldoutCompositePrimitiveIds;
-	for (const TSoftObjectPtr<UPrimitiveComponent>& PrimitivePtr : CompositePrimitives)
+	for (const TWeakObjectPtr<UPrimitiveComponent>& PrimitivePtr : CompositePrimitives)
 	{
-		if (PrimitivePtr.IsValid())
+		TStrongObjectPtr<UPrimitiveComponent> Primitive = PrimitivePtr.Pin();
+		if (Primitive.IsValid())
 		{
-			const FPrimitiveComponentId PrimId = PrimitivePtr->GetPrimitiveSceneId();
+			const FPrimitiveComponentId PrimId = Primitive->GetPrimitiveSceneId();
 
 			if (InView.ShowOnlyPrimitives.IsSet())
 			{
 				if (InView.ShowOnlyPrimitives.GetValue().Contains(PrimId))
 				{
-					HoldoutCompositePrimitiveIds.Add(PrimitivePtr->GetPrimitiveSceneId());
+					HoldoutCompositePrimitiveIds.Add(Primitive->GetPrimitiveSceneId());
 				}
 			}
 			else if (!InView.HiddenPrimitives.Contains(PrimId))
 			{
-				HoldoutCompositePrimitiveIds.Add(PrimitivePtr->GetPrimitiveSceneId());
+				HoldoutCompositePrimitiveIds.Add(Primitive->GetPrimitiveSceneId());
 			}
 		}
 	}
@@ -487,6 +504,21 @@ FScreenPassTexture FHoldoutCompositeSceneViewExtension::PostProcessCompositePass
 
 	FScreenPassTexture CustomRenderPassTexture = FScreenPassTexture(GetCustomRenderPassTexture(GraphBuilder, InView));
 
+	{
+		FScopeLock Lock(&MaterialsCS);
+		for (const TWeakObjectPtr<UMaterial>& MaterialPtr : CompositeProcessMaterials)
+		{
+			TStrongObjectPtr<UMaterial> Material = MaterialPtr.Pin();
+			if (Material.IsValid())
+			{
+				FPostProcessMaterialInputs PostProcessMaterialInputs;
+				PostProcessMaterialInputs.SetInput(GraphBuilder, EPostProcessMaterialInput::SceneColor, CustomRenderPassTexture);
+				PostProcessMaterialInputs.SceneTextures = Inputs.SceneTextures; //Note: these are invalid as they come from the main render, but validation-required.
+				CustomRenderPassTexture = AddPostProcessMaterialPass(GraphBuilder, InView, PostProcessMaterialInputs, Material.Get());
+			}
+		}
+	}
+
 	if (bApplyFXAA)
 	{
 		// We gamma-encode the result so that FXAA can operate on non-linear colors
diff --git a/Engine/Plugins/Experimental/Compositing/HoldoutComposite/Source/HoldoutComposite/Private/HoldoutCompositeSceneViewExtension.h b/Engine/Plugins/Experimental/Compositing/HoldoutComposite/Source/HoldoutComposite/Private/HoldoutCompositeSceneViewExtension.h
--- a/Engine/Plugins/Experimental/Compositing/HoldoutComposite/Source/HoldoutComposite/Private/HoldoutCompositeSceneViewExtension.h
+++ b/Engine/Plugins/Experimental/Compositing/HoldoutComposite/Source/HoldoutComposite/Private/HoldoutCompositeSceneViewExtension.h
@@ -7,10 +7,11 @@
 #include "RendererInterface.h"
 #include "Containers/ContainersFwd.h"
 
+enum class EHoldoutCompositePassLocation;
 class FViewInfo;
 struct FScreenPassRenderTarget;
 struct FScreenPassTexture;
-enum class EHoldoutCompositePassLocation;
+class UMaterial;
 
 class FHoldoutCompositeSceneViewExtension : public FWorldSceneViewExtension
 {
@@ -19,10 +20,16 @@ public:
 	~FHoldoutCompositeSceneViewExtension();
 
 	/* Register primitives for compositing. */
-	void RegisterPrimitives(TArrayView<TSoftObjectPtr<UPrimitiveComponent>> InPrimitiveComponents, bool bInHoldoutState=true);
+	void RegisterPrimitives_GameThread(TArrayView<UPrimitiveComponent*> InPrimitiveComponents, bool bInHoldoutState=true);
 
 	/* Unregister primitives for compositing. */
-	void UnregisterPrimitives(TArrayView<TSoftObjectPtr<UPrimitiveComponent>> InPrimitiveComponents, bool bInHoldoutState=false);
+	void UnregisterPrimitives_GameThread(TArrayView<UPrimitiveComponent*> InPrimitiveComponents, bool bInHoldoutState=false);
+
+	/* Set additional post-process materials to be applied on the separate composited render. */
+	void SetCompositePostProcessMaterials_GameThread(TArrayView<UMaterial*> InMaterials);
+
+	/* Remove all post-process materials to be applied on the separate composited render. */
+	void ResetCompositePostProcessMaterials_GameThread();
 
 	/* Called by the custom render pass to store its view render target for this frame. */
 	template<typename T>
@@ -55,7 +62,7 @@ private:
 	FScreenPassTexture PostProcessCompositePass_RenderThread(FRDGBuilder& GraphBuilder, const FSceneView& InView, const FPostProcessMaterialInputs& Inputs);
 
 	// Collection of primitives to render as a custom render pass and composite after post-processing.
-	TSet<TSoftObjectPtr<UPrimitiveComponent>> CompositePrimitives;
+	TSet<TWeakObjectPtr<UPrimitiveComponent>> CompositePrimitives;
 
 	// Custom render pass render targets for each active view
 	TMap<uint32, TRefCountPtr<IPooledRenderTarget>> CustomRenderTargetPerView_RenderThread;
@@ -66,5 +73,11 @@ private:
 	std::atomic_bool bApplyTSRLensDistortion;
 	std::atomic_bool bApplyFXAA;
 	std::atomic<EHoldoutCompositePassLocation> CompositePassLocation;
+
+	/** Post-process material(s) to be applied on the separate composited render.*/
+	TArray<TWeakObjectPtr<UMaterial>> CompositeProcessMaterials;
+
+	/** Critical section for composite post-process materials. */
+	FCriticalSection MaterialsCS;
 };
 
diff --git a/Engine/Plugins/Experimental/Compositing/HoldoutComposite/Source/HoldoutComposite/Private/HoldoutCompositeSubsystem.cpp b/Engine/Plugins/Experimental/Compositing/HoldoutComposite/Source/HoldoutComposite/Private/HoldoutCompositeSubsystem.cpp
--- a/Engine/Plugins/Experimental/Compositing/HoldoutComposite/Source/HoldoutComposite/Private/HoldoutCompositeSubsystem.cpp
+++ b/Engine/Plugins/Experimental/Compositing/HoldoutComposite/Source/HoldoutComposite/Private/HoldoutCompositeSubsystem.cpp
@@ -64,35 +64,51 @@ void UHoldoutCompositeSubsystem::Deinitialize()
 	Super::Deinitialize();
 }
 
-void UHoldoutCompositeSubsystem::RegisterPrimitive(TSoftObjectPtr<UPrimitiveComponent> InPrimitiveComponent, bool bInHoldoutState)
+void UHoldoutCompositeSubsystem::RegisterPrimitive(UPrimitiveComponent* InPrimitiveComponent, bool bInHoldoutState)
 {
-	TArray<TSoftObjectPtr<UPrimitiveComponent>> PrimitiveComponents = { InPrimitiveComponent };
+	TArray<UPrimitiveComponent*> PrimitiveComponents = { InPrimitiveComponent };
 
 	RegisterPrimitives(PrimitiveComponents, bInHoldoutState);
 }
 
-void UHoldoutCompositeSubsystem::RegisterPrimitives(TArrayView<TSoftObjectPtr<UPrimitiveComponent>> InPrimitiveComponents, bool bInHoldoutState)
+void UHoldoutCompositeSubsystem::RegisterPrimitives(TArrayView<UPrimitiveComponent*> InPrimitiveComponents, bool bInHoldoutState)
 {
 	const bool bValidSettings = ValidateProjectSettings();
 
 	if (bValidSettings && HoldoutCompositeViewExtension.IsValid())
 	{
-		HoldoutCompositeViewExtension->RegisterPrimitives(MoveTemp(InPrimitiveComponents), bInHoldoutState);
+		HoldoutCompositeViewExtension->RegisterPrimitives_GameThread(MoveTemp(InPrimitiveComponents), bInHoldoutState);
 	}
 }
 
-void UHoldoutCompositeSubsystem::UnregisterPrimitive(TSoftObjectPtr<UPrimitiveComponent> InPrimitiveComponent, bool bInHoldoutState)
+void UHoldoutCompositeSubsystem::UnregisterPrimitive(UPrimitiveComponent* InPrimitiveComponent, bool bInHoldoutState)
 {
-	TArray<TSoftObjectPtr<UPrimitiveComponent>> PrimitiveComponents = { InPrimitiveComponent };
-	
+	TArray<UPrimitiveComponent*> PrimitiveComponents = { InPrimitiveComponent };
+
 	UnregisterPrimitives(PrimitiveComponents, bInHoldoutState);
 }
 
-void UHoldoutCompositeSubsystem::UnregisterPrimitives(TArrayView<TSoftObjectPtr<UPrimitiveComponent>> InPrimitiveComponents, bool bInHoldoutState)
+void UHoldoutCompositeSubsystem::UnregisterPrimitives(TArrayView<UPrimitiveComponent*> InPrimitiveComponents, bool bInHoldoutState)
+{
+	if (HoldoutCompositeViewExtension.IsValid())
+	{
+		HoldoutCompositeViewExtension->UnregisterPrimitives_GameThread(MoveTemp(InPrimitiveComponents), bInHoldoutState);
+	}
+}
+
+void UHoldoutCompositeSubsystem::SetCompositePostProcessMaterials(TArray<UMaterial*> InMaterials)
+{
+	if (HoldoutCompositeViewExtension.IsValid())
+	{
+		HoldoutCompositeViewExtension->SetCompositePostProcessMaterials_GameThread(InMaterials);
+	}
+}
+
+void UHoldoutCompositeSubsystem::ResetCompositePostProcessMaterials()
 {
 	if (HoldoutCompositeViewExtension.IsValid())
 	{
-		HoldoutCompositeViewExtension->UnregisterPrimitives(MoveTemp(InPrimitiveComponents), bInHoldoutState);
+		HoldoutCompositeViewExtension->ResetCompositePostProcessMaterials_GameThread();
 	}
 }
 
diff --git a/Engine/Plugins/Experimental/Compositing/HoldoutComposite/Source/HoldoutComposite/Public/HoldoutCompositeSubsystem.h b/Engine/Plugins/Experimental/Compositing/HoldoutComposite/Source/HoldoutComposite/Public/HoldoutCompositeSubsystem.h
--- a/Engine/Plugins/Experimental/Compositing/HoldoutComposite/Source/HoldoutComposite/Public/HoldoutCompositeSubsystem.h
+++ b/Engine/Plugins/Experimental/Compositing/HoldoutComposite/Source/HoldoutComposite/Public/HoldoutCompositeSubsystem.h
@@ -4,12 +4,12 @@
 
 #include "MaterialShared.h"
 #include "Subsystems/WorldSubsystem.h"
-#include "UObject/SoftObjectPtr.h"
 
 #include "HoldoutCompositeSubsystem.generated.h"
 
 class FHoldoutCompositeSceneViewExtension;
 class SNotificationItem;
+class UPrimitiveComponent;
 class URendererSettings;
 
 /**
@@ -30,17 +30,25 @@ public:
 
 	/* Register a single primitive for compositing. */
 	UFUNCTION(BlueprintCallable, Category = "Holdout Composite")
-	HOLDOUTCOMPOSITE_API void RegisterPrimitive(TSoftObjectPtr<UPrimitiveComponent> InPrimitiveComponent, bool bInHoldoutState=true);
+	HOLDOUTCOMPOSITE_API void RegisterPrimitive(UPrimitiveComponent* InPrimitiveComponent, bool bInHoldoutState=true);
 
 	/* Register multiple primitives for compositing. */
-	HOLDOUTCOMPOSITE_API void RegisterPrimitives(TArrayView<TSoftObjectPtr<UPrimitiveComponent>> InPrimitiveComponents, bool bInHoldoutState=true);
+	HOLDOUTCOMPOSITE_API void RegisterPrimitives(TArrayView<UPrimitiveComponent*> InPrimitiveComponents, bool bInHoldoutState=true);
 
 	/* Unregister a single primitive from compositing. */
 	UFUNCTION(BlueprintCallable, Category = "Holdout Composite")
-	HOLDOUTCOMPOSITE_API void UnregisterPrimitive(TSoftObjectPtr<UPrimitiveComponent> InPrimitiveComponent, bool bInHoldoutState=false);
+	HOLDOUTCOMPOSITE_API void UnregisterPrimitive(UPrimitiveComponent* InPrimitiveComponent, bool bInHoldoutState=false);
 
 	/* Unregister multiple primitives from compositing. */
-	HOLDOUTCOMPOSITE_API void UnregisterPrimitives(TArrayView<TSoftObjectPtr<UPrimitiveComponent>> InPrimitiveComponents, bool bInHoldoutState=false);
+	HOLDOUTCOMPOSITE_API void UnregisterPrimitives(TArrayView<UPrimitiveComponent*> InPrimitiveComponents, bool bInHoldoutState=false);
+
+	/* Set additional post-process materials to be applied on the separate composited render. */
+	UFUNCTION(BlueprintCallable, Category = "Holdout Composite")
+	HOLDOUTCOMPOSITE_API void SetCompositePostProcessMaterials(TArray<UMaterial*> InMaterials);
+
+	/* Remove all post-process materials to be applied on the separate composited render. */
+	UFUNCTION(BlueprintCallable, Category = "Holdout Composite")
+	HOLDOUTCOMPOSITE_API void ResetCompositePostProcessMaterials();
 
 private:
 
