HoldoutComposite: Minor fix from previous CL, replace material with material interface on (temporary) subsystem functions.

diff --git a/Engine/Plugins/Experimental/Compositing/HoldoutComposite/Source/HoldoutComposite/Private/HoldoutCompositeSceneViewExtension.cpp b/Engine/Plugins/Experimental/Compositing/HoldoutComposite/Source/HoldoutComposite/Private/HoldoutCompositeSceneViewExtension.cpp
--- a/Engine/Plugins/Experimental/Compositing/HoldoutComposite/Source/HoldoutComposite/Private/HoldoutCompositeSceneViewExtension.cpp
+++ b/Engine/Plugins/Experimental/Compositing/HoldoutComposite/Source/HoldoutComposite/Private/HoldoutCompositeSceneViewExtension.cpp
@@ -273,7 +273,7 @@ void FHoldoutCompositeSceneViewExtension::UnregisterPrimitives_GameThread(TArray
 	}
 }
 
-void FHoldoutCompositeSceneViewExtension::SetCompositePostProcessMaterials_GameThread(TArrayView<UMaterial*> InMaterials)
+void FHoldoutCompositeSceneViewExtension::SetCompositePostProcessMaterials_GameThread(TArrayView<UMaterialInterface*> InMaterials)
 {
 	check(IsInGameThread());
 
@@ -506,9 +506,9 @@ FScreenPassTexture FHoldoutCompositeSceneViewExtension::PostProcessCompositePass
 
 	{
 		FScopeLock Lock(&MaterialsCS);
-		for (const TWeakObjectPtr<UMaterial>& MaterialPtr : CompositeProcessMaterials)
+		for (const TWeakObjectPtr<UMaterialInterface>& MaterialPtr : CompositeProcessMaterials)
 		{
-			TStrongObjectPtr<UMaterial> Material = MaterialPtr.Pin();
+			TStrongObjectPtr<UMaterialInterface> Material = MaterialPtr.Pin();
 			if (Material.IsValid())
 			{
 				FPostProcessMaterialInputs PostProcessMaterialInputs;
diff --git a/Engine/Plugins/Experimental/Compositing/HoldoutComposite/Source/HoldoutComposite/Private/HoldoutCompositeSceneViewExtension.h b/Engine/Plugins/Experimental/Compositing/HoldoutComposite/Source/HoldoutComposite/Private/HoldoutCompositeSceneViewExtension.h
--- a/Engine/Plugins/Experimental/Compositing/HoldoutComposite/Source/HoldoutComposite/Private/HoldoutCompositeSceneViewExtension.h
+++ b/Engine/Plugins/Experimental/Compositing/HoldoutComposite/Source/HoldoutComposite/Private/HoldoutCompositeSceneViewExtension.h
@@ -11,7 +11,7 @@ enum class EHoldoutCompositePassLocation;
 class FViewInfo;
 struct FScreenPassRenderTarget;
 struct FScreenPassTexture;
-class UMaterial;
+class UMaterialInterface;
 
 class FHoldoutCompositeSceneViewExtension : public FWorldSceneViewExtension
 {
@@ -26,7 +26,7 @@ public:
 	void UnregisterPrimitives_GameThread(TArrayView<UPrimitiveComponent*> InPrimitiveComponents, bool bInHoldoutState=false);
 
 	/* Set additional post-process materials to be applied on the separate composited render. */
-	void SetCompositePostProcessMaterials_GameThread(TArrayView<UMaterial*> InMaterials);
+	void SetCompositePostProcessMaterials_GameThread(TArrayView<UMaterialInterface*> InMaterials);
 
 	/* Remove all post-process materials to be applied on the separate composited render. */
 	void ResetCompositePostProcessMaterials_GameThread();
@@ -75,7 +75,7 @@ private:
 	std::atomic<EHoldoutCompositePassLocation> CompositePassLocation;
 
 	/** Post-process material(s) to be applied on the separate composited render.*/
-	TArray<TWeakObjectPtr<UMaterial>> CompositeProcessMaterials;
+	TArray<TWeakObjectPtr<UMaterialInterface>> CompositeProcessMaterials;
 
 	/** Critical section for composite post-process materials. */
 	FCriticalSection MaterialsCS;
diff --git a/Engine/Plugins/Experimental/Compositing/HoldoutComposite/Source/HoldoutComposite/Private/HoldoutCompositeSubsystem.cpp b/Engine/Plugins/Experimental/Compositing/HoldoutComposite/Source/HoldoutComposite/Private/HoldoutCompositeSubsystem.cpp
--- a/Engine/Plugins/Experimental/Compositing/HoldoutComposite/Source/HoldoutComposite/Private/HoldoutCompositeSubsystem.cpp
+++ b/Engine/Plugins/Experimental/Compositing/HoldoutComposite/Source/HoldoutComposite/Private/HoldoutCompositeSubsystem.cpp
@@ -96,7 +96,7 @@ void UHoldoutCompositeSubsystem::UnregisterPrimitives(TArrayView<UPrimitiveCompo
 	}
 }
 
-void UHoldoutCompositeSubsystem::SetCompositePostProcessMaterials(TArray<UMaterial*> InMaterials)
+void UHoldoutCompositeSubsystem::SetCompositePostProcessMaterials(TArray<UMaterialInterface*> InMaterials)
 {
 	if (HoldoutCompositeViewExtension.IsValid())
 	{
diff --git a/Engine/Plugins/Experimental/Compositing/HoldoutComposite/Source/HoldoutComposite/Public/HoldoutCompositeSubsystem.h b/Engine/Plugins/Experimental/Compositing/HoldoutComposite/Source/HoldoutComposite/Public/HoldoutCompositeSubsystem.h
--- a/Engine/Plugins/Experimental/Compositing/HoldoutComposite/Source/HoldoutComposite/Public/HoldoutCompositeSubsystem.h
+++ b/Engine/Plugins/Experimental/Compositing/HoldoutComposite/Source/HoldoutComposite/Public/HoldoutCompositeSubsystem.h
@@ -44,7 +44,7 @@ public:
 
 	/* Set additional post-process materials to be applied on the separate composited render. */
 	UFUNCTION(BlueprintCallable, Category = "Holdout Composite")
-	HOLDOUTCOMPOSITE_API void SetCompositePostProcessMaterials(TArray<UMaterial*> InMaterials);
+	HOLDOUTCOMPOSITE_API void SetCompositePostProcessMaterials(TArray<UMaterialInterface*> InMaterials);
 
 	/* Remove all post-process materials to be applied on the separate composited render. */
 	UFUNCTION(BlueprintCallable, Category = "Holdout Composite")
