HoldoutComposite: New actor that regroups composite settings, meshes and materials for ease of use.

diff --git a/Engine/Plugins/Experimental/Compositing/HoldoutComposite/Source/HoldoutComposite/Private/HoldoutComposite.cpp b/Engine/Plugins/Experimental/Compositing/HoldoutComposite/Source/HoldoutComposite/Private/HoldoutComposite.cpp
new file mode 100644
--- /dev/null
+++ b/Engine/Plugins/Experimental/Compositing/HoldoutComposite/Source/HoldoutComposite/Private/HoldoutComposite.cpp
@@ -0,0 +1,153 @@
+// Copyright Epic Games, Inc. All Rights Reserved.
+
+#include "HoldoutComposite.h"
+#include "HoldoutCompositeSubsystem.h"
+#include "HoldoutCompositeModule.h"
+
+#include "Components/PrimitiveComponent.h"
+#include "Engine/Engine.h"
+#include "Engine/World.h"
+#include "UObject/UObjectIterator.h"
+
+#include UE_INLINE_GENERATED_CPP_BY_NAME(HoldoutComposite)
+
+namespace UE
+{
+	namespace HoldoutComposite
+	{
+		namespace Private
+		{
+			TArray<UPrimitiveComponent*> GetPrimitives(TArrayView<TObjectPtr<AActor>> InActors)
+			{
+				TArray<UPrimitiveComponent*> OutPrimitiveComponents;
+
+				const UHoldoutCompositePluginSettings* Settings = GetDefault<UHoldoutCompositePluginSettings>();
+				check(Settings);
+
+				for (const TObjectPtr<AActor>& Actor : InActors)
+				{
+					TInlineComponentArray<UPrimitiveComponent*> PrimComponents(Actor);
+
+					for (UPrimitiveComponent* PrimitiveComponent : PrimComponents)
+					{
+						if (Settings->IsAllowedPrimitiveClass(PrimitiveComponent))
+						{
+							OutPrimitiveComponents.Add(PrimitiveComponent);
+						}
+					}
+				}
+
+				return OutPrimitiveComponents;
+			}
+		}
+	}
+}
+
+void AHoldoutComposite::PostActorCreated()
+{
+	Super::PostActorCreated();
+
+#if WITH_EDITOR
+	int32 HoldoutCompositeCount = 0;
+	for (TObjectIterator<AHoldoutComposite> It; It; ++It)
+	{
+		HoldoutCompositeCount++;
+	}
+
+	if (HoldoutCompositeCount > 1)
+	{
+		if (GEngine)
+		{
+			GEngine->AddOnScreenDebugMessage(-1, 5.0f, FColor::Red, TEXT("No more than one holdout composite actor is allowed."));
+		}
+
+		UE_LOG(LogHoldoutComposite, Error, TEXT("No more than one holdout composite actor is allowed, disabling. Please remove."));
+		bIsEnabled = false;
+	}
+#endif // WITH_EDITOR
+}
+
+void AHoldoutComposite::PostRegisterAllComponents()
+{
+	Super::PostRegisterAllComponents();
+
+	RegisterCompositeImpl();
+}
+
+void AHoldoutComposite::PostUnregisterAllComponents()
+{
+	Super::PostUnregisterAllComponents();
+
+	UnregisterCompositeImpl();
+}
+
+void AHoldoutComposite::BeginDestroy()
+{
+	Super::BeginDestroy();
+}
+
+#if WITH_EDITOR
+void AHoldoutComposite::PreEditChange(FProperty* PropertyThatWillChange)
+{
+	Super::PreEditChange(PropertyThatWillChange);
+}
+
+void AHoldoutComposite::PostEditChangeProperty(FPropertyChangedEvent& PropertyChangedEvent)
+{
+	Super::PostEditChangeProperty(PropertyChangedEvent);
+}
+#endif
+
+bool AHoldoutComposite::IsEnabled() const
+{
+	return bIsEnabled;
+}
+
+void AHoldoutComposite::SetEnabled(bool bInEnabled)
+{
+	bIsEnabled = bInEnabled;
+
+	if (bIsEnabled)
+	{
+		RegisterCompositeImpl();
+	}
+	else
+	{
+		UnregisterCompositeImpl();
+	}
+}
+
+void AHoldoutComposite::RegisterCompositeImpl()
+{
+	using namespace UE::HoldoutComposite::Private;
+
+	if (!bIsEnabled)
+	{
+		return;
+	}
+
+	UHoldoutCompositeSubsystem* Subsystem = UWorld::GetSubsystem<UHoldoutCompositeSubsystem>(GetWorld());
+	if (IsValid(Subsystem))
+	{
+		if (bOverrideDefaultSettings)
+		{
+			Subsystem->OverrideCompositeSettings(SettingsOverride);
+		}
+		Subsystem->RegisterPrimitives(GetPrimitives(CompositeMeshActors));
+		Subsystem->SetCompositePostProcessMaterials(CompositePostProcessMaterials);
+	}
+}
+
+void AHoldoutComposite::UnregisterCompositeImpl()
+{
+	using namespace UE::HoldoutComposite::Private;
+
+	UHoldoutCompositeSubsystem* Subsystem = UWorld::GetSubsystem<UHoldoutCompositeSubsystem>(GetWorld());
+	if (IsValid(Subsystem))
+	{
+		Subsystem->ResetCompositePostProcessMaterials();
+		Subsystem->UnregisterPrimitives(GetPrimitives(CompositeMeshActors));
+		Subsystem->ResetCompositeSettings();
+	}
+}
+
diff --git a/Engine/Plugins/Experimental/Compositing/HoldoutComposite/Source/HoldoutComposite/Private/HoldoutCompositeComponent.cpp b/Engine/Plugins/Experimental/Compositing/HoldoutComposite/Source/HoldoutComposite/Private/HoldoutCompositeComponent.cpp
--- a/Engine/Plugins/Experimental/Compositing/HoldoutComposite/Source/HoldoutComposite/Private/HoldoutCompositeComponent.cpp
+++ b/Engine/Plugins/Experimental/Compositing/HoldoutComposite/Source/HoldoutComposite/Private/HoldoutCompositeComponent.cpp
@@ -15,58 +15,48 @@
 
 #define LOCTEXT_NAMESPACE "HoldoutComposite"
 
-namespace
+namespace UE
 {
-	// Ensure that the primitive class is allowed by checking against plugin settings.
-	bool IsAllowedPrimitiveClass(const UPrimitiveComponent* InPrimitiveComponent)
+	namespace HoldoutComposite
 	{
-		const UHoldoutCompositeSettings* Settings = GetDefault<UHoldoutCompositeSettings>();
-		if (ensure(Settings))
+		namespace Private
 		{
-			for (const FSoftClassPath& ObjectClass : Settings->DisabledPrimitiveClasses)
+			// Find list of primitives from the parent component and its children. We need to traverse children to support objects such as UText3DComponent.
+			static TArray<UPrimitiveComponent*> FindPrimitiveComponents(USceneComponent* InParentComponent)
 			{
-				UClass* TransactionClass = ObjectClass.TryLoadClass<UObject>();
-				if (TransactionClass && InPrimitiveComponent->IsA(TransactionClass))
-				{
-					return false;
-				}
-			}
-		}
+				const UHoldoutCompositePluginSettings* PluginSettings = GetDefault<UHoldoutCompositePluginSettings>();
+				check(PluginSettings);
 
-		return true;
-	}
+				TArray<UPrimitiveComponent*> OutPrimitiveComponents;
 
-	// Find list of primitives from the parent component and its children. We need to traverse children to support objects such as UText3DComponent.
-	static TArray<UPrimitiveComponent*> FindPrimitiveComponents(USceneComponent* InParentComponent)
-	{
-		TArray<UPrimitiveComponent*> OutPrimitiveComponents;
-
-		if (IsValid(InParentComponent))
-		{
-			if (UPrimitiveComponent* PrimitiveComponent = Cast<UPrimitiveComponent>(InParentComponent))
-			{
-				if (IsAllowedPrimitiveClass(PrimitiveComponent))
+				if (IsValid(InParentComponent))
 				{
-					OutPrimitiveComponents.Add(PrimitiveComponent);
-				}
-			}
+					if (UPrimitiveComponent* PrimitiveComponent = Cast<UPrimitiveComponent>(InParentComponent))
+					{
+						if (PluginSettings->IsAllowedPrimitiveClass(PrimitiveComponent))
+						{
+							OutPrimitiveComponents.Add(PrimitiveComponent);
+						}
+					}
 
-			TArray<USceneComponent*> ParentChildComponents;
-			InParentComponent->GetChildrenComponents(true /*bIncludeAllDescendants*/, ParentChildComponents);
+					TArray<USceneComponent*> ParentChildComponents;
+					InParentComponent->GetChildrenComponents(true /*bIncludeAllDescendants*/, ParentChildComponents);
 
-			for (USceneComponent* ParentChild : ParentChildComponents)
-			{
-				if (UPrimitiveComponent* PrimitiveComponent = Cast<UPrimitiveComponent>(ParentChild))
-				{
-					if (IsAllowedPrimitiveClass(PrimitiveComponent))
+					for (USceneComponent* ParentChild : ParentChildComponents)
 					{
-						OutPrimitiveComponents.Add(PrimitiveComponent);
+						if (UPrimitiveComponent* PrimitiveComponent = Cast<UPrimitiveComponent>(ParentChild))
+						{
+							if (PluginSettings->IsAllowedPrimitiveClass(PrimitiveComponent))
+							{
+								OutPrimitiveComponents.Add(PrimitiveComponent);
+							}
+						}
 					}
 				}
+
+				return OutPrimitiveComponents;
 			}
 		}
-
-		return OutPrimitiveComponents;
 	}
 }
 
@@ -106,7 +96,7 @@ void UHoldoutCompositeComponent::OnAttachmentChanged()
 	USceneComponent* SceneComponent = GetAttachParent();
 	if (IsValid(SceneComponent))
 	{
-		TArray<UPrimitiveComponent*> ParentPrimitives = FindPrimitiveComponents(SceneComponent);
+		TArray<UPrimitiveComponent*> ParentPrimitives = UE::HoldoutComposite::Private::FindPrimitiveComponents(SceneComponent);
 		if (ParentPrimitives.IsEmpty())
 		{
 #if WITH_EDITOR
@@ -154,7 +144,7 @@ void UHoldoutCompositeComponent::RegisterCompositeImpl()
 	}
 
 	UHoldoutCompositeSubsystem* Subsystem = UWorld::GetSubsystem<UHoldoutCompositeSubsystem>(GetWorld());
-	TArray<UPrimitiveComponent*> ParentPrimitives = FindPrimitiveComponents(GetAttachParent());
+	TArray<UPrimitiveComponent*> ParentPrimitives = UE::HoldoutComposite::Private::FindPrimitiveComponents(GetAttachParent());
 
 	if (IsValid(Subsystem) && !ParentPrimitives.IsEmpty())
 	{
@@ -165,7 +155,7 @@ void UHoldoutCompositeComponent::RegisterCompositeImpl()
 void UHoldoutCompositeComponent::UnregisterCompositeImpl()
 {
 	UHoldoutCompositeSubsystem* Subsystem = UWorld::GetSubsystem<UHoldoutCompositeSubsystem>(GetWorld());
-	TArray<UPrimitiveComponent*> ParentPrimitives = FindPrimitiveComponents(GetAttachParent());
+	TArray<UPrimitiveComponent*> ParentPrimitives = UE::HoldoutComposite::Private::FindPrimitiveComponents(GetAttachParent());
 
 	if (IsValid(Subsystem) && !ParentPrimitives.IsEmpty())
 	{
diff --git a/Engine/Plugins/Experimental/Compositing/HoldoutComposite/Source/HoldoutComposite/Private/HoldoutCompositeSceneViewExtension.cpp b/Engine/Plugins/Experimental/Compositing/HoldoutComposite/Source/HoldoutComposite/Private/HoldoutCompositeSceneViewExtension.cpp
--- a/Engine/Plugins/Experimental/Compositing/HoldoutComposite/Source/HoldoutComposite/Private/HoldoutCompositeSceneViewExtension.cpp
+++ b/Engine/Plugins/Experimental/Compositing/HoldoutComposite/Source/HoldoutComposite/Private/HoldoutCompositeSceneViewExtension.cpp
@@ -2,8 +2,8 @@
 
 #include "HoldoutCompositeSceneViewExtension.h"
 
-#include "HoldoutCompositeSettings.h"
 #include "HoldoutCompositeModule.h"
+#include "HoldoutComposite.h"
 
 #include "Components/PrimitiveComponent.h"
 #include "Containers/Map.h"
@@ -224,18 +224,18 @@ private:
 
 FHoldoutCompositeSceneViewExtension::FHoldoutCompositeSceneViewExtension(const FAutoRegister& AutoReg, UWorld* InWorld)
 	: FWorldSceneViewExtension(AutoReg, InWorld)
-	, bFollowSceneExposure(false)
-	, bSupportSSR(true)
-	, bApplyTSRLensDistortion(true)
-	, CompositePassLocation(EHoldoutCompositePassLocation::AfterTonemap)
 {
+	const UHoldoutCompositePluginSettings* PluginSettings = GetDefault<UHoldoutCompositePluginSettings>();
+	check(PluginSettings);
+
+	CompositeSettings = CompositeSettings_RenderThread = PluginSettings->DefaultCompositeSettings;
 }
 
 FHoldoutCompositeSceneViewExtension::~FHoldoutCompositeSceneViewExtension()
 {
 }
 
-void FHoldoutCompositeSceneViewExtension::RegisterPrimitives_GameThread(TArrayView<UPrimitiveComponent*> InPrimitiveComponents, bool bInHoldoutState)
+void FHoldoutCompositeSceneViewExtension::RegisterPrimitives_GameThread(const TArray<UPrimitiveComponent*>& InPrimitiveComponents, bool bInHoldoutState)
 {
 	check(IsInGameThread());
 
@@ -254,7 +254,7 @@ void FHoldoutCompositeSceneViewExtension::RegisterPrimitives_GameThread(TArrayVi
 	}
 }
 
-void FHoldoutCompositeSceneViewExtension::UnregisterPrimitives_GameThread(TArrayView<UPrimitiveComponent*> InPrimitiveComponents, bool bInHoldoutState)
+void FHoldoutCompositeSceneViewExtension::UnregisterPrimitives_GameThread(const TArray<UPrimitiveComponent*>& InPrimitiveComponents, bool bInHoldoutState)
 {
 	check(IsInGameThread());
 
@@ -273,12 +273,30 @@ void FHoldoutCompositeSceneViewExtension::UnregisterPrimitives_GameThread(TArray
 	}
 }
 
-void FHoldoutCompositeSceneViewExtension::SetCompositePostProcessMaterials_GameThread(TArrayView<UMaterialInterface*> InMaterials)
+void FHoldoutCompositeSceneViewExtension::OverrideCompositeSettings_GameThread(const FHoldoutCompositeSettings& InSettings)
+{
+	check(IsInGameThread());
+
+	UpdateSettings_GameThread(InSettings);
+}
+
+void FHoldoutCompositeSceneViewExtension::ResetCompositeSettings_GameThread()
+{
+	check(IsInGameThread());
+
+	const UHoldoutCompositePluginSettings* PluginSettings = GetDefault<UHoldoutCompositePluginSettings>();
+	check(PluginSettings);
+
+	UpdateSettings_GameThread(PluginSettings->DefaultCompositeSettings);
+}
+
+void FHoldoutCompositeSceneViewExtension::SetCompositePostProcessMaterials_GameThread(const TArray<UMaterialInterface*>& InMaterials)
 {
 	check(IsInGameThread());
 
 	FScopeLock Lock(&MaterialsCS);
-	CompositeProcessMaterials = InMaterials;
+	CompositeProcessMaterials.Reset();
+	CompositeProcessMaterials.Append(InMaterials);
 }
 
 void FHoldoutCompositeSceneViewExtension::ResetCompositePostProcessMaterials_GameThread()
@@ -303,22 +321,12 @@ bool FHoldoutCompositeSceneViewExtension::IsActiveThisFrame_Internal(const FScen
 
 int32 FHoldoutCompositeSceneViewExtension::GetPriority() const
 {
-	const UHoldoutCompositeSettings* Settings = GetDefault<UHoldoutCompositeSettings>();
-
-	return Settings ? Settings->CompositePassPriority : 0;
+	return GetSettings().CompositePassPriority;
 }
 
 void FHoldoutCompositeSceneViewExtension::SetupViewFamily(FSceneViewFamily& InViewFamily)
 {
-	const UHoldoutCompositeSettings* Settings = GetDefault<UHoldoutCompositeSettings>();
-	if (Settings != nullptr)
-	{
-		bFollowSceneExposure = Settings->bFollowSceneExposure;
-		bSupportSSR	= Settings->bSupportSSR;
-		bApplyTSRLensDistortion = Settings->bApplyTSRLensDistortion;
-		bApplyFXAA = Settings->bApplyFXAA;
-		CompositePassLocation = Settings->CompositePassLocation;
-	}
+
 }
 
 void FHoldoutCompositeSceneViewExtension::SetupView(FSceneViewFamily& InViewFamily, FSceneView& InView)
@@ -376,13 +384,11 @@ void FHoldoutCompositeSceneViewExtension::SetupView(FSceneViewFamily& InViewFami
 		InView
 	);
 
-	const UHoldoutCompositeSettings* Settings = GetDefault<UHoldoutCompositeSettings>();
-
 	FSceneInterface::FCustomRenderPassRendererInput PassInput{};
 	PassInput.EngineShowFlags = InViewFamily.EngineShowFlags;
 	PassInput.EngineShowFlags.DisableFeaturesForUnlit();
 	PassInput.EngineShowFlags.SetTranslucency(true);
-	PassInput.EngineShowFlags.SetUnlitViewmode(Settings ? Settings->bSupportUnlitViewmode : false);
+	PassInput.EngineShowFlags.SetUnlitViewmode(GetSettings().bEnableUnlitViewmode);
 	PassInput.EngineShowFlags.SetAllowPrimitiveAlphaHoldout(false);
 
 	// Note: Incoming view location is invalid for scene captures
@@ -408,20 +414,46 @@ void FHoldoutCompositeSceneViewExtension::SubscribeToPostProcessingPass(EPostPro
 		return;
 	}
 	
-	if (PassId == EPostProcessingPass::SSRInput && bSupportSSR.load())
+	if (PassId == EPostProcessingPass::SSRInput && GetSettings().bUpdateSSR)
 	{
 		InOutPassCallbacks.Add(FAfterPassCallbackDelegate::CreateRaw(this, &FHoldoutCompositeSceneViewExtension::PostProcessPassSSRInput_RenderThread));
 	}
-	else if (PassId == EPostProcessingPass::Tonemap && CompositePassLocation.load() == EHoldoutCompositePassLocation::AfterTonemap)
+	else if (PassId == EPostProcessingPass::Tonemap && GetSettings().CompositePassLocation == EHoldoutCompositePassLocation::AfterTonemap)
 	{
 		InOutPassCallbacks.Add(FAfterPassCallbackDelegate::CreateRaw(this, &FHoldoutCompositeSceneViewExtension::PostProcessCompositePass_RenderThread));
 	}
-	else if (PassId == EPostProcessingPass::VisualizeDepthOfField && CompositePassLocation.load() == EHoldoutCompositePassLocation::EndOfPostProcessing)
+	else if (PassId == EPostProcessingPass::VisualizeDepthOfField && GetSettings().CompositePassLocation == EHoldoutCompositePassLocation::EndOfPostProcessing)
 	{
 		InOutPassCallbacks.Add(FAfterPassCallbackDelegate::CreateRaw(this, &FHoldoutCompositeSceneViewExtension::PostProcessCompositePass_RenderThread));
 	}
 }
 
+void FHoldoutCompositeSceneViewExtension::UpdateSettings_GameThread(const FHoldoutCompositeSettings& InSettings)
+{
+	CompositeSettings = InSettings;
+	
+	ENQUEUE_RENDER_COMMAND(CopyHoldoutCompositeSettings)(
+		[CompositeSettingsRT = InSettings, WeakThis = AsWeak()](FRHICommandListImmediate& RHICmdList)
+		{
+			const TSharedPtr<FHoldoutCompositeSceneViewExtension> SVE = StaticCastSharedPtr<FHoldoutCompositeSceneViewExtension>(WeakThis.Pin());
+			if (SVE.IsValid())
+			{
+				SVE->CompositeSettings_RenderThread = CompositeSettingsRT;
+			}
+		});
+}
+
+const FHoldoutCompositeSettings& FHoldoutCompositeSceneViewExtension::GetSettings() const
+{
+	if (IsInRenderingThread())
+	{
+		return CompositeSettings_RenderThread;
+	}
+	
+	checkf(IsInGameThread(), TEXT("Composite settings should only be access from the render thread and the game thread."));
+	return CompositeSettings;
+}
+
 FRDGTextureRef FHoldoutCompositeSceneViewExtension::GetCustomRenderPassTexture(FRDGBuilder& GraphBuilder, const FSceneView& InView) const
 {
 	FRDGTextureRef CustomRenderPassTexture = GSystemTextures.GetBlackAlphaOneDummy(GraphBuilder);
@@ -476,7 +508,7 @@ FScreenPassTexture FHoldoutCompositeSceneViewExtension::PostProcessPassSSRInput_
 	FHoldoutCompositeSSRInputShader::FParameters* PassParameters = GraphBuilder.AllocParameters<FHoldoutCompositeSSRInputShader::FParameters>();
 	constexpr bool bIsSceneColorUndistorted = true;
 	PassParameters->Common = BuildCommonCompositeParameters(GraphBuilder, InView, SceneColor, FScreenPassTexture(CustomRenderPassTexture), Output, bIsSceneColorUndistorted);
-	PassParameters->LastGlobalExposure = bFollowSceneExposure.load() ? InView.GetLastEyeAdaptationExposure() : 1.0f;
+	PassParameters->LastGlobalExposure = GetSettings().bFollowSceneExposure ? InView.GetLastEyeAdaptationExposure() : 1.0f;
 	PassParameters->RenderTargets[0] = Output.GetRenderTargetBinding();
 
 	TShaderMapRef<FHoldoutCompositeSSRInputShader> PixelShader(GlobalShaderMap);
@@ -514,12 +546,18 @@ FScreenPassTexture FHoldoutCompositeSceneViewExtension::PostProcessCompositePass
 				FPostProcessMaterialInputs PostProcessMaterialInputs;
 				PostProcessMaterialInputs.SetInput(GraphBuilder, EPostProcessMaterialInput::SceneColor, CustomRenderPassTexture);
 				PostProcessMaterialInputs.SceneTextures = Inputs.SceneTextures; //Note: these are invalid as they come from the main render, but validation-required.
+
+				// Create our own output override to maintain UAV access
+				FRDGTextureDesc OutputDesc = CustomRenderPassTexture.Texture->Desc;
+				OutputDesc.Flags |= ETextureCreateFlags::RenderTargetable;
+				PostProcessMaterialInputs.OverrideOutput = FScreenPassRenderTarget(GraphBuilder.CreateTexture(OutputDesc, TEXT("PostProcessMaterial")), CustomRenderPassTexture.ViewRect, InView.GetOverwriteLoadAction());
+
 				CustomRenderPassTexture = AddPostProcessMaterialPass(GraphBuilder, InView, PostProcessMaterialInputs, Material.Get());
 			}
 		}
 	}
 
-	if (bApplyFXAA)
+	if (GetSettings().bApplyFXAA)
 	{
 		// We gamma-encode the result so that FXAA can operate on non-linear colors
 		HoldoutComposite::AddGammaPass(GraphBuilder, CustomRenderPassTexture.Texture, InView.GetFeatureLevel(), 1.0f / HoldoutComposite::DefaultGamma);
@@ -547,7 +585,7 @@ FScreenPassTexture FHoldoutCompositeSceneViewExtension::PostProcessCompositePass
 		Encodings.Y = static_cast<uint32>(ESceneColorSourceEncoding::sRGB);
 	}
 	// FXAA received gamma-encoded colors, so we must linearize before compositing.
-	Encodings.Z = static_cast<uint32>(bApplyFXAA);
+	Encodings.Z = static_cast<uint32>(GetSettings().bApplyFXAA);
 
 	FScreenPassRenderTarget Output = Inputs.OverrideOutput;
 	if (!Output.IsValid())
@@ -558,19 +596,19 @@ FScreenPassTexture FHoldoutCompositeSceneViewExtension::PostProcessCompositePass
 	// Compositing pass
 	{
 		FHoldoutCompositeFinalShader::FPermutationDomain PermutationVector;
-		PermutationVector.Set<FHoldoutCompositeFinalShader::FUseGlobalExposure>(bFollowSceneExposure.load());
+		PermutationVector.Set<FHoldoutCompositeFinalShader::FUseGlobalExposure>(GetSettings().bFollowSceneExposure);
 
 		FHoldoutCompositeFinalShader::FParameters* PassParameters = GraphBuilder.AllocParameters<FHoldoutCompositeFinalShader::FParameters>();
 		PassParameters->View = InView.ViewUniformBuffer;
 		PassParameters->Common = BuildCommonCompositeParameters(GraphBuilder, InView, SceneColor, CustomRenderPassTexture, Output);
-		if (bFollowSceneExposure.load())
+		if (GetSettings().bFollowSceneExposure)
 		{
 			FRDGBufferRef EyeAdaptationBuffer = GraphBuilder.RegisterExternalBuffer(InView.GetEyeAdaptationBuffer(), ERDGBufferFlags::MultiFrame);
 			PassParameters->EyeAdaptationBuffer = GraphBuilder.CreateSRV(EyeAdaptationBuffer);
 		}
 		PassParameters->Encodings = Encodings;
 		PassParameters->DisplayGamma = FVector3f(Family->RenderTarget->GetDisplayGamma(), 1.0f / Family->RenderTarget->GetDisplayGamma(), HoldoutComposite::DefaultGamma);
-		PassParameters->bApplyTSRLensDistortion = bApplyTSRLensDistortion;
+		PassParameters->bApplyTSRLensDistortion = GetSettings().bApplyTSRLensDistortion;
 		PassParameters->RenderTargets[0] = Output.GetRenderTargetBinding();
 
 		TShaderMapRef<FHoldoutCompositeFinalShader> PixelShader(GlobalShaderMap, PermutationVector);
diff --git a/Engine/Plugins/Experimental/Compositing/HoldoutComposite/Source/HoldoutComposite/Private/HoldoutCompositeSceneViewExtension.h b/Engine/Plugins/Experimental/Compositing/HoldoutComposite/Source/HoldoutComposite/Private/HoldoutCompositeSceneViewExtension.h
--- a/Engine/Plugins/Experimental/Compositing/HoldoutComposite/Source/HoldoutComposite/Private/HoldoutCompositeSceneViewExtension.h
+++ b/Engine/Plugins/Experimental/Compositing/HoldoutComposite/Source/HoldoutComposite/Private/HoldoutCompositeSceneViewExtension.h
@@ -7,10 +7,13 @@
 #include "RendererInterface.h"
 #include "Containers/ContainersFwd.h"
 
+#include "HoldoutCompositeSettings.h"
+
 enum class EHoldoutCompositePassLocation;
 class FViewInfo;
 struct FScreenPassRenderTarget;
 struct FScreenPassTexture;
+struct FHoldoutCompositeSettings;
 class UMaterialInterface;
 
 class FHoldoutCompositeSceneViewExtension : public FWorldSceneViewExtension
@@ -20,13 +23,19 @@ public:
 	~FHoldoutCompositeSceneViewExtension();
 
 	/* Register primitives for compositing. */
-	void RegisterPrimitives_GameThread(TArrayView<UPrimitiveComponent*> InPrimitiveComponents, bool bInHoldoutState=true);
+	void RegisterPrimitives_GameThread(const TArray<UPrimitiveComponent*>& InPrimitiveComponents, bool bInHoldoutState=true);
 
 	/* Unregister primitives for compositing. */
-	void UnregisterPrimitives_GameThread(TArrayView<UPrimitiveComponent*> InPrimitiveComponents, bool bInHoldoutState=false);
+	void UnregisterPrimitives_GameThread(const TArray<UPrimitiveComponent*>& InPrimitiveComponents, bool bInHoldoutState=false);
+
+	/* Override the composite settings. */
+	void OverrideCompositeSettings_GameThread(const FHoldoutCompositeSettings& InSettings);
+
+	/* Reset composite settings back to plugin settings default. */
+	void ResetCompositeSettings_GameThread();
 
 	/* Set additional post-process materials to be applied on the separate composited render. */
-	void SetCompositePostProcessMaterials_GameThread(TArrayView<UMaterialInterface*> InMaterials);
+	void SetCompositePostProcessMaterials_GameThread(const TArray<UMaterialInterface*>& InMaterials);
 
 	/* Remove all post-process materials to be applied on the separate composited render. */
 	void ResetCompositePostProcessMaterials_GameThread();
@@ -53,6 +62,12 @@ protected:
 	virtual bool IsActiveThisFrame_Internal(const FSceneViewExtensionContext& Context) const override;
 	//~ End ISceneViewExtension Interface
 
+	/* Update composite settings. */
+	void UpdateSettings_GameThread(const FHoldoutCompositeSettings& InSettings);
+
+	/* Thread-safe getter for composite settings. */
+	const FHoldoutCompositeSettings& GetSettings() const;
+
 private:
 
 	FRDGTextureRef GetCustomRenderPassTexture(FRDGBuilder& GraphBuilder, const FSceneView& InView) const;
@@ -67,12 +82,9 @@ private:
 	// Custom render pass render targets for each active view
 	TMap<uint32, TRefCountPtr<IPooledRenderTarget>> CustomRenderTargetPerView_RenderThread;
 
-	// Duplicate flags from plugin settings for thread-safety
-	std::atomic_bool bFollowSceneExposure;
-	std::atomic_bool bSupportSSR;
-	std::atomic_bool bApplyTSRLensDistortion;
-	std::atomic_bool bApplyFXAA;
-	std::atomic<EHoldoutCompositePassLocation> CompositePassLocation;
+	// Composite settings and its render-thread copy
+	FHoldoutCompositeSettings CompositeSettings;
+	FHoldoutCompositeSettings CompositeSettings_RenderThread;
 
 	/** Post-process material(s) to be applied on the separate composited render.*/
 	TArray<TWeakObjectPtr<UMaterialInterface>> CompositeProcessMaterials;
diff --git a/Engine/Plugins/Experimental/Compositing/HoldoutComposite/Source/HoldoutComposite/Private/HoldoutCompositeSettings.cpp b/Engine/Plugins/Experimental/Compositing/HoldoutComposite/Source/HoldoutComposite/Private/HoldoutCompositeSettings.cpp
--- a/Engine/Plugins/Experimental/Compositing/HoldoutComposite/Source/HoldoutComposite/Private/HoldoutCompositeSettings.cpp
+++ b/Engine/Plugins/Experimental/Compositing/HoldoutComposite/Source/HoldoutComposite/Private/HoldoutCompositeSettings.cpp
@@ -2,32 +2,52 @@
 
 #include "HoldoutCompositeSettings.h"
 
+#include "Components/PrimitiveComponent.h"
+
 #include UE_INLINE_GENERATED_CPP_BY_NAME(HoldoutCompositeSettings)
 
-UHoldoutCompositeSettings::UHoldoutCompositeSettings()
+
+FHoldoutCompositeSettings::FHoldoutCompositeSettings()
 	: bFollowSceneExposure(false)
-	, bSupportSSR(true)
-	, bSupportUnlitViewmode(true)
+	, bUpdateSSR(true)
+	, bEnableUnlitViewmode(true)
 	, bApplyTSRLensDistortion(true)
 	, bApplyFXAA(false)
-	, DisabledPrimitiveClasses()
 	, CompositePassLocation(EHoldoutCompositePassLocation::AfterTonemap)
 	, CompositePassPriority(100 + 1 /* OPENCOLORIO_SCENE_VIEW_EXTENSION_PRIORITY + 1 */)
 { }
 
-FName UHoldoutCompositeSettings::GetCategoryName() const
+UHoldoutCompositePluginSettings::UHoldoutCompositePluginSettings()
+	: DisabledPrimitiveClasses()
+	, DefaultCompositeSettings()
+{ }
+
+FName UHoldoutCompositePluginSettings::GetCategoryName() const
 {
 	return TEXT("Plugins");
 }
 
 #if WITH_EDITOR
-FText UHoldoutCompositeSettings::GetSectionText() const
+FText UHoldoutCompositePluginSettings::GetSectionText() const
 {
 	return NSLOCTEXT("HoldoutCompositeSettings", "HoldoutCompositeSettingsSection", "Holdout Composite");
 }
 
-FName UHoldoutCompositeSettings::GetSectionName() const
+FName UHoldoutCompositePluginSettings::GetSectionName() const
 {
 	return TEXT("Holdout Composite");
 }
+bool UHoldoutCompositePluginSettings::IsAllowedPrimitiveClass(const UPrimitiveComponent* InPrimitiveComponent) const
+{
+	for (const FSoftClassPath& ObjectClass : DisabledPrimitiveClasses)
+	{
+		UClass* TransactionClass = ObjectClass.TryLoadClass<UObject>();
+		if (TransactionClass && InPrimitiveComponent->IsA(TransactionClass))
+		{
+			return false;
+		}
+	}
+
+	return true;
+}
 #endif
diff --git a/Engine/Plugins/Experimental/Compositing/HoldoutComposite/Source/HoldoutComposite/Private/HoldoutCompositeSubsystem.cpp b/Engine/Plugins/Experimental/Compositing/HoldoutComposite/Source/HoldoutComposite/Private/HoldoutCompositeSubsystem.cpp
--- a/Engine/Plugins/Experimental/Compositing/HoldoutComposite/Source/HoldoutComposite/Private/HoldoutCompositeSubsystem.cpp
+++ b/Engine/Plugins/Experimental/Compositing/HoldoutComposite/Source/HoldoutComposite/Private/HoldoutCompositeSubsystem.cpp
@@ -71,13 +71,13 @@ void UHoldoutCompositeSubsystem::RegisterPrimitive(UPrimitiveComponent* InPrimit
 	RegisterPrimitives(PrimitiveComponents, bInHoldoutState);
 }
 
-void UHoldoutCompositeSubsystem::RegisterPrimitives(TArrayView<UPrimitiveComponent*> InPrimitiveComponents, bool bInHoldoutState)
+void UHoldoutCompositeSubsystem::RegisterPrimitives(const TArray<UPrimitiveComponent*>& InPrimitiveComponents, bool bInHoldoutState)
 {
 	const bool bValidSettings = ValidateProjectSettings();
 
 	if (bValidSettings && HoldoutCompositeViewExtension.IsValid())
 	{
-		HoldoutCompositeViewExtension->RegisterPrimitives_GameThread(MoveTemp(InPrimitiveComponents), bInHoldoutState);
+		HoldoutCompositeViewExtension->RegisterPrimitives_GameThread(InPrimitiveComponents, bInHoldoutState);
 	}
 }
 
@@ -88,15 +88,15 @@ void UHoldoutCompositeSubsystem::UnregisterPrimitive(UPrimitiveComponent* InPrim
 	UnregisterPrimitives(PrimitiveComponents, bInHoldoutState);
 }
 
-void UHoldoutCompositeSubsystem::UnregisterPrimitives(TArrayView<UPrimitiveComponent*> InPrimitiveComponents, bool bInHoldoutState)
+void UHoldoutCompositeSubsystem::UnregisterPrimitives(const TArray<UPrimitiveComponent*>& InPrimitiveComponents, bool bInHoldoutState)
 {
 	if (HoldoutCompositeViewExtension.IsValid())
 	{
-		HoldoutCompositeViewExtension->UnregisterPrimitives_GameThread(MoveTemp(InPrimitiveComponents), bInHoldoutState);
+		HoldoutCompositeViewExtension->UnregisterPrimitives_GameThread(InPrimitiveComponents, bInHoldoutState);
 	}
 }
 
-void UHoldoutCompositeSubsystem::SetCompositePostProcessMaterials(TArray<UMaterialInterface*> InMaterials)
+void UHoldoutCompositeSubsystem::SetCompositePostProcessMaterials(const TArray<UMaterialInterface*>& InMaterials)
 {
 	if (HoldoutCompositeViewExtension.IsValid())
 	{
@@ -112,6 +112,22 @@ void UHoldoutCompositeSubsystem::ResetCompositePostProcessMaterials()
 	}
 }
 
+void UHoldoutCompositeSubsystem::OverrideCompositeSettings(const FHoldoutCompositeSettings& InSettings)
+{
+	if (HoldoutCompositeViewExtension.IsValid())
+	{
+		HoldoutCompositeViewExtension->OverrideCompositeSettings_GameThread(InSettings);
+	}
+}
+
+void UHoldoutCompositeSubsystem::ResetCompositeSettings()
+{
+	if (HoldoutCompositeViewExtension.IsValid())
+	{
+		HoldoutCompositeViewExtension->ResetCompositeSettings_GameThread();
+	}
+}
+
 bool UHoldoutCompositeSubsystem::ValidateProjectSettings()
 {
 	URendererSettings* RendererSettings = GetMutableDefault<URendererSettings>();
diff --git a/Engine/Plugins/Experimental/Compositing/HoldoutComposite/Source/HoldoutComposite/Public/HoldoutComposite.h b/Engine/Plugins/Experimental/Compositing/HoldoutComposite/Source/HoldoutComposite/Public/HoldoutComposite.h
new file mode 100644
--- /dev/null
+++ b/Engine/Plugins/Experimental/Compositing/HoldoutComposite/Source/HoldoutComposite/Public/HoldoutComposite.h
@@ -0,0 +1,73 @@
+// Copyright Epic Games, Inc. All Rights Reserved.
+
+#pragma once
+
+#include "CoreTypes.h"
+#include "GameFramework/Actor.h"
+#include "HoldoutCompositeSettings.h"
+
+#include "HoldoutComposite.generated.h"
+
+class UPrimitiveComponent;
+class UMaterialInterface;
+
+UCLASS()
+class HOLDOUTCOMPOSITE_API AHoldoutComposite : public AActor
+{
+	GENERATED_BODY()
+
+public:
+
+	//~ Begin AActor Interface
+	virtual void PostActorCreated();
+	virtual void PostRegisterAllComponents() override;
+	virtual void PostUnregisterAllComponents() override;
+	virtual void BeginDestroy() override;
+	//~ End AActor Interface
+
+#if WITH_EDITOR
+	virtual void PreEditChange(FProperty* PropertyThatWillChange) override;
+	virtual void PostEditChangeProperty(FPropertyChangedEvent& PropertyChangedEvent) override;
+#endif // WITH_EDITOR
+
+	/* Get the enabled state of the component. */
+	UFUNCTION(BlueprintGetter)
+	bool IsEnabled() const;
+
+	/* Set the enabled state of the component. */
+	UFUNCTION(BlueprintSetter)
+	void SetEnabled(bool bInEnabled);
+
+private:
+
+	/* Whether or not the component activates the composite. */
+	UPROPERTY(EditAnywhere, BlueprintGetter = IsEnabled, BlueprintSetter = SetEnabled, Category = HoldoutComposite, meta = (AllowPrivateAccess = true))
+	bool bIsEnabled = true;
+
+public:
+
+	/* Actors whose primitive (mesh) components will be marked as primitive alpha holdout, rendered separately and composited back under scene color after post-processing. */
+	UPROPERTY(EditAnywhere, BlueprintReadOnly, Category = HoldoutComposite)
+	TArray<TObjectPtr<AActor>> CompositeMeshActors;
+
+	/** Post-process material(s) to be applied on the separate composited render.*/
+	UPROPERTY(EditAnywhere, BlueprintReadOnly, Category = HoldoutComposite)
+	TArray<TObjectPtr<UMaterialInterface>> CompositePostProcessMaterials;
+
+	/** Post-process material(s) to be applied on the separate composited render.*/
+	UPROPERTY(EditAnywhere, BlueprintReadOnly, Category = HoldoutComposite)
+	bool bOverrideDefaultSettings = false;
+
+	/** Post process settings to use for this holdout composite. */
+	UPROPERTY(EditAnywhere, BlueprintReadOnly, Category = HoldoutComposite, meta=(EditCondition=bOverrideDefaultSettings))
+	FHoldoutCompositeSettings SettingsOverride = {};
+
+private:
+
+	/* Private implementation of the register method. */
+	void RegisterCompositeImpl();
+
+	/* Private implementation of the unregister method. */
+	void UnregisterCompositeImpl();
+};
+
diff --git a/Engine/Plugins/Experimental/Compositing/HoldoutComposite/Source/HoldoutComposite/Public/HoldoutCompositeSettings.h b/Engine/Plugins/Experimental/Compositing/HoldoutComposite/Source/HoldoutComposite/Public/HoldoutCompositeSettings.h
--- a/Engine/Plugins/Experimental/Compositing/HoldoutComposite/Source/HoldoutComposite/Public/HoldoutCompositeSettings.h
+++ b/Engine/Plugins/Experimental/Compositing/HoldoutComposite/Source/HoldoutComposite/Public/HoldoutCompositeSettings.h
@@ -9,6 +9,9 @@
 
 #include "HoldoutCompositeSettings.generated.h"
 
+class UPrimitiveComponent;
+
+/* Post-processing location of the main holdout composite pass. */
 UENUM()
 enum class EHoldoutCompositePassLocation
 {
@@ -17,55 +20,76 @@ enum class EHoldoutCompositePassLocation
 };
 
 /**
- * Settings for the HoldoutComposite module.
+ * Settings for the active composite (default values can be overriden by the composite actor).
  */
-UCLASS(config = HoldoutComposite, defaultconfig, meta = (DisplayName = "Holdout Composite"), MinimalAPI)
-class UHoldoutCompositeSettings : public UDeveloperSettings
+USTRUCT(BlueprintType)
+struct HOLDOUTCOMPOSITE_API FHoldoutCompositeSettings
 {
 	GENERATED_BODY()
 
-public:
-	UHoldoutCompositeSettings();
-
-	//~ Begin UDeveloperSettings interface
-	virtual FName GetCategoryName() const;
-#if WITH_EDITOR
-	virtual FText GetSectionText() const override;
-	virtual FName GetSectionName() const override;
-#endif
-	//~ End UDeveloperSettings interface
+	/** Constructor */
+	FHoldoutCompositeSettings();
 
-public:
 	/** When enabled, the view global exposure is applied onto the separate render when composited. */
-	UPROPERTY(config, EditAnywhere, Category = General)
+	UPROPERTY(config, EditAnywhere, Category = HoldoutComposite)
 	bool bFollowSceneExposure;
 
 	/** When enabled, the separate composited render is also used to update screen-space reflections. */
-	UPROPERTY(config, EditAnywhere, Category = General)
-	bool bSupportSSR;
+	UPROPERTY(config, EditAnywhere, Category = HoldoutComposite)
+	bool bUpdateSSR;
 
 	/** Enables the development shader debug feature that routes the Base Color output to Emissive for the separate render. Non-shipping PC build only. */
-	UPROPERTY(config, EditAnywhere, Category = General)
-	bool bSupportUnlitViewmode;
+	UPROPERTY(config, EditAnywhere, Category = HoldoutComposite)
+	bool bEnableUnlitViewmode;
 
 	/** When enabled, the TSR lens distortion (r.TSR.LensDistortion / Scene View Extension (Experimental) on lens files) is also applied on the separate composited render. */
-	UPROPERTY(config, EditAnywhere, Category = General, DisplayName="Apply TSR Lens Distortion")
+	UPROPERTY(config, EditAnywhere, Category = HoldoutComposite, DisplayName = "Apply TSR Lens Distortion")
 	bool bApplyTSRLensDistortion;
 
 	/** When enabled, FXAA is applied on the separate composited render. */
-	UPROPERTY(config, EditAnywhere, Category = General, DisplayName = "Apply FXAA")
+	UPROPERTY(config, EditAnywhere, Category = HoldoutComposite, DisplayName = "Apply FXAA")
 	bool bApplyFXAA;
 
-	/** Primitive component classes that do not support the holdout composite.*/
-	UPROPERTY(config, EditAnywhere, Category = General)
-	TArray<FSoftClassPath> DisabledPrimitiveClasses;
-
 	/** Composite (scene view extension) pass location, which defaults to "After Tonemap". "End Of Post Processing" may be used to push the composite pass after additional post-process materials (such as the regular lens distortion) and OCIO. */
-	UPROPERTY(config, EditAnywhere, Category = General, AdvancedDisplay)
+	UPROPERTY(config, EditAnywhere, Category = HoldoutComposite, AdvancedDisplay)
 	EHoldoutCompositePassLocation CompositePassLocation;
 
 	/** Composite (scene view extension) pass priority, which defaults to before OpenColorIO. */
-	UPROPERTY(config, EditAnywhere, Category = General, AdvancedDisplay)
+	UPROPERTY(config, EditAnywhere, Category = HoldoutComposite, AdvancedDisplay)
 	int32 CompositePassPriority;
 };
 
+/**
+ * Settings for the HoldoutComposite module.
+ */
+UCLASS(config = Engine, meta = (DisplayName = "Holdout Composite"), MinimalAPI)
+class UHoldoutCompositePluginSettings : public UDeveloperSettings
+{
+	GENERATED_BODY()
+
+public:
+	UHoldoutCompositePluginSettings();
+
+	//~ Begin UDeveloperSettings interface
+	virtual FName GetCategoryName() const;
+#if WITH_EDITOR
+	virtual FText GetSectionText() const override;
+	virtual FName GetSectionName() const override;
+#endif
+	//~ End UDeveloperSettings interface
+
+public:
+
+	// Ensure that the primitive class is allowed by checking against plugin settings.
+	bool IsAllowedPrimitiveClass(const UPrimitiveComponent* InPrimitiveComponent) const;
+
+	/** Primitive component classes that do not support the holdout composite.*/
+	UPROPERTY(config, EditAnywhere, Category = General)
+	TArray<FSoftClassPath> DisabledPrimitiveClasses;
+
+	/** Default holdout composite settings.*/
+	UPROPERTY(config, EditAnywhere, Category = General, meta = (ShowOnlyInnerProperties))
+	FHoldoutCompositeSettings DefaultCompositeSettings;
+};
+
+
diff --git a/Engine/Plugins/Experimental/Compositing/HoldoutComposite/Source/HoldoutComposite/Public/HoldoutCompositeSubsystem.h b/Engine/Plugins/Experimental/Compositing/HoldoutComposite/Source/HoldoutComposite/Public/HoldoutCompositeSubsystem.h
--- a/Engine/Plugins/Experimental/Compositing/HoldoutComposite/Source/HoldoutComposite/Public/HoldoutCompositeSubsystem.h
+++ b/Engine/Plugins/Experimental/Compositing/HoldoutComposite/Source/HoldoutComposite/Public/HoldoutCompositeSubsystem.h
@@ -8,6 +8,7 @@
 #include "HoldoutCompositeSubsystem.generated.h"
 
 class FHoldoutCompositeSceneViewExtension;
+struct FHoldoutCompositeSettings;
 class SNotificationItem;
 class UPrimitiveComponent;
 class URendererSettings;
@@ -33,23 +34,31 @@ public:
 	HOLDOUTCOMPOSITE_API void RegisterPrimitive(UPrimitiveComponent* InPrimitiveComponent, bool bInHoldoutState=true);
 
 	/* Register multiple primitives for compositing. */
-	HOLDOUTCOMPOSITE_API void RegisterPrimitives(TArrayView<UPrimitiveComponent*> InPrimitiveComponents, bool bInHoldoutState=true);
+	HOLDOUTCOMPOSITE_API void RegisterPrimitives(const TArray<UPrimitiveComponent*>& InPrimitiveComponents, bool bInHoldoutState=true);
 
 	/* Unregister a single primitive from compositing. */
 	UFUNCTION(BlueprintCallable, Category = "Holdout Composite")
 	HOLDOUTCOMPOSITE_API void UnregisterPrimitive(UPrimitiveComponent* InPrimitiveComponent, bool bInHoldoutState=false);
 
 	/* Unregister multiple primitives from compositing. */
-	HOLDOUTCOMPOSITE_API void UnregisterPrimitives(TArrayView<UPrimitiveComponent*> InPrimitiveComponents, bool bInHoldoutState=false);
+	HOLDOUTCOMPOSITE_API void UnregisterPrimitives(const TArray<UPrimitiveComponent*>& InPrimitiveComponents, bool bInHoldoutState=false);
 
 	/* Set additional post-process materials to be applied on the separate composited render. */
 	UFUNCTION(BlueprintCallable, Category = "Holdout Composite")
-	HOLDOUTCOMPOSITE_API void SetCompositePostProcessMaterials(TArray<UMaterialInterface*> InMaterials);
+	HOLDOUTCOMPOSITE_API void SetCompositePostProcessMaterials(const TArray<UMaterialInterface*>& InMaterials);
 
 	/* Remove all post-process materials to be applied on the separate composited render. */
 	UFUNCTION(BlueprintCallable, Category = "Holdout Composite")
 	HOLDOUTCOMPOSITE_API void ResetCompositePostProcessMaterials();
 
+	/* Override project default composite settings for the active holdout composite passes. */
+	UFUNCTION(BlueprintCallable, Category = "Holdout Composite")
+	HOLDOUTCOMPOSITE_API void OverrideCompositeSettings(const FHoldoutCompositeSettings& InSettings);
+
+	/* Revert composite settings back to project defaults. */
+	UFUNCTION(BlueprintCallable, Category = "Holdout Composite")
+	HOLDOUTCOMPOSITE_API void ResetCompositeSettings();
+
 private:
 
 	/* Returns true if the (renderer) project settings are correctly enabled for the composite to be active. */
