HoldoutComposite: Fix include build error.

diff --git a/Engine/Plugins/Experimental/Compositing/HoldoutComposite/Source/HoldoutComposite/Private/HoldoutComposite.cpp b/Engine/Plugins/Experimental/Compositing/HoldoutComposite/Source/HoldoutComposite/Private/HoldoutComposite.cpp
--- a/Engine/Plugins/Experimental/Compositing/HoldoutComposite/Source/HoldoutComposite/Private/HoldoutComposite.cpp
+++ b/Engine/Plugins/Experimental/Compositing/HoldoutComposite/Source/HoldoutComposite/Private/HoldoutComposite.cpp
@@ -3,6 +3,7 @@
 #include "HoldoutComposite.h"
 #include "HoldoutCompositeSubsystem.h"
 #include "HoldoutCompositeModule.h"
+#include "HoldoutCompositeSettings.h"
 
 #include "Components/PrimitiveComponent.h"
 #include "Engine/Engine.h"
