HoldoutComposite: Adds hidden in editor support to new actor, fixing behavior for the visibility toggle.

diff --git a/Engine/Plugins/Experimental/Compositing/HoldoutComposite/Source/HoldoutComposite/Private/HoldoutComposite.cpp b/Engine/Plugins/Experimental/Compositing/HoldoutComposite/Source/HoldoutComposite/Private/HoldoutComposite.cpp
--- a/Engine/Plugins/Experimental/Compositing/HoldoutComposite/Source/HoldoutComposite/Private/HoldoutComposite.cpp
+++ b/Engine/Plugins/Experimental/Compositing/HoldoutComposite/Source/HoldoutComposite/Private/HoldoutComposite.cpp
@@ -97,6 +97,20 @@ void AHoldoutComposite::PostEditChangeProperty(FPropertyChangedEvent& PropertyCh
 {
 	Super::PostEditChangeProperty(PropertyChangedEvent);
 }
+
+void AHoldoutComposite::SetIsTemporarilyHiddenInEditor(bool bIsHidden)
+{
+	Super::SetIsTemporarilyHiddenInEditor(bIsHidden);
+
+	if (IsTemporarilyHiddenInEditor())
+	{
+		UnregisterCompositeImpl();
+	}
+	else
+	{
+		RegisterCompositeImpl();
+	}
+}
 #endif
 
 bool AHoldoutComposite::IsEnabled() const
diff --git a/Engine/Plugins/Experimental/Compositing/HoldoutComposite/Source/HoldoutComposite/Public/HoldoutComposite.h b/Engine/Plugins/Experimental/Compositing/HoldoutComposite/Source/HoldoutComposite/Public/HoldoutComposite.h
--- a/Engine/Plugins/Experimental/Compositing/HoldoutComposite/Source/HoldoutComposite/Public/HoldoutComposite.h
+++ b/Engine/Plugins/Experimental/Compositing/HoldoutComposite/Source/HoldoutComposite/Public/HoldoutComposite.h
@@ -28,6 +28,7 @@ public:
 #if WITH_EDITOR
 	virtual void PreEditChange(FProperty* PropertyThatWillChange) override;
 	virtual void PostEditChangeProperty(FPropertyChangedEvent& PropertyChangedEvent) override;
+	virtual void SetIsTemporarilyHiddenInEditor(bool bIsHidden) override;
 #endif // WITH_EDITOR
 
 	/* Get the enabled state of the component. */
