HoldoutComposite: Allow direct use a composite texture override (in screen-space), instead of relying on the custom render pass.

* Converted gamma compute shader to pixel shader to stay compatible with external textures usually only used with PPMs.
* Also adds composite view user flags setting.

diff --git a/Engine/Plugins/Experimental/Compositing/HoldoutComposite/Shaders/Private/HoldoutCompositeFinal.usf b/Engine/Plugins/Experimental/Compositing/HoldoutComposite/Shaders/Private/HoldoutCompositeFinal.usf
--- a/Engine/Plugins/Experimental/Compositing/HoldoutComposite/Shaders/Private/HoldoutCompositeFinal.usf
+++ b/Engine/Plugins/Experimental/Compositing/HoldoutComposite/Shaders/Private/HoldoutCompositeFinal.usf
@@ -6,14 +6,14 @@
 #include "/Engine/Private/LensDistortion.ush"
 
 SCREEN_PASS_TEXTURE_VIEWPORT(Input)
-SCREEN_PASS_TEXTURE_VIEWPORT(Custom)
+SCREEN_PASS_TEXTURE_VIEWPORT(Composite)
 SCREEN_PASS_TEXTURE_VIEWPORT(Output)
 
 Texture2D InputTexture;
 SamplerState InputSampler;
 
-Texture2D CustomTexture;
-SamplerState CustomSampler;
+Texture2D CompositeTexture;
+SamplerState CompositeSampler;
 
 Texture2D<float2> DistortingDisplacementTexture;
 SamplerState DistortingDisplacementSampler;
@@ -25,6 +25,7 @@ StructuredBuffer<float4> EyeAdaptationBuffer;
 uint3 Encodings;
 float3 DisplayGamma;
 uint bApplyTSRLensDistortion;
+uint bIsCustomRenderPassTexture;
 
 float SafeRcp(float x)
 {
@@ -52,14 +53,14 @@ void MainPS(
     float4 SceneColor = Texture2DSample(InputTexture, InputSampler, InputUV);
 
     const float2 UndistortedViewportUV = ApplyLensDistortionOnViewportUV(UndistortingDisplacementTexture, UndistortingDisplacementSampler, DistortedViewportUV);
-    const float2 CustomViewportUV = select(bApplyTSRLensDistortion > 0, UndistortedViewportUV, DistortedViewportUV);
+    const float2 CompositeViewportUV = select(bApplyTSRLensDistortion > 0, UndistortedViewportUV, DistortedViewportUV);
 
-    const float2 CustomUV = (CustomViewportUV * Custom_ViewportSize + Custom_ViewportMin) * Custom_ExtentInverse;
-    float4 CustomColor = Texture2DSample(CustomTexture, CustomSampler, CustomUV);
+    const float2 CompositeUV = (CompositeViewportUV * Composite_ViewportSize + Composite_ViewportMin) * Composite_ExtentInverse;
+    float4 CompositeColor = Texture2DSample(CompositeTexture, CompositeSampler, CompositeUV);
 
 #if USE_GLOBAL_EXPOSURE
     const float GlobalExposure = EyeAdaptationBuffer[0].x;
-    CustomColor.rgb *= GlobalExposure;
+    CompositeColor.rgb *= GlobalExposure;
 #endif
 
     BRANCH
@@ -72,18 +73,24 @@ void MainPS(
         SceneColor.rgb = sRGBToLinear(SceneColor.rgb);
     }
 
-    // Optionally linearize custom color
+    // Optionally linearize composite color
     if(Encodings.z > 0)
     {
-        CustomColor.rgb = pow(CustomColor.rgb, DisplayGamma.z);
+        CompositeColor.rgb = pow(CompositeColor.rgb, DisplayGamma.z);
     }
 
-    // The alpha comes from the scene color holdout mask, so we unpremultiply the custom render to get its opaque colors. (No-op on opaque geometry.)
-    CustomColor = InvertAlpha(Unpremultiply(InvertAlpha(CustomColor)));
+	// Custom render pass texture has inverted pre-multiplied alpha (like the renderer)
+	if (bIsCustomRenderPassTexture > 0)
+    {
+        CompositeColor = InvertAlpha(CompositeColor);
+        // The alpha comes from the scene color holdout mask, so we unpremultiply the composite render to get its opaque colors. (No-op on opaque geometry.)
+        CompositeColor = Unpremultiply(CompositeColor);
+    }
+    CompositeColor = InvertAlpha(CompositeColor);
 
     // TODO: HDR cases will need to be handled explicitely.
-    OutColor.rgb = SceneColor.rgb + CustomColor.rgb * SceneColor.a;
-    OutColor.a = CustomColor.a * SceneColor.a;
+    OutColor.rgb = SceneColor.rgb + CompositeColor.rgb * SceneColor.a;
+    OutColor.a = CompositeColor.a * SceneColor.a;
 
     BRANCH
     if(Encodings.y == 1)
diff --git a/Engine/Plugins/Experimental/Compositing/HoldoutComposite/Shaders/Private/HoldoutCompositeGamma.usf b/Engine/Plugins/Experimental/Compositing/HoldoutComposite/Shaders/Private/HoldoutCompositeGamma.usf
--- a/Engine/Plugins/Experimental/Compositing/HoldoutComposite/Shaders/Private/HoldoutCompositeGamma.usf
+++ b/Engine/Plugins/Experimental/Compositing/HoldoutComposite/Shaders/Private/HoldoutCompositeGamma.usf
@@ -2,26 +2,15 @@
 
 #include "/Engine/Private/Common.ush"
 
-#ifndef TILE_SIZE
-    #error TILE_SIZE must be defined
-#endif
-
 Texture2D InputTexture;
-RWTexture2D<float4> RWOutputTexture;
-uint2 Dimensions;
+SamplerState InputSampler;
 float GammaPower;
 
-[numthreads(TILE_SIZE, TILE_SIZE, 1)]
-void MainCS(
-    uint2 DispatchThreadId : SV_DispatchThreadID)
+void MainPS(
+    noperspective float4 InUVAndScreenPos : TEXCOORD0,
+    out float4 OutColor : SV_Target0
+    )
 {
-    if (any(DispatchThreadId >= Dimensions))
-    {
-        return;
-    }
-
-    float4 TexelColor = InputTexture[DispatchThreadId];
-    TexelColor.rgb = pow(TexelColor.rgb, abs(GammaPower));
-
-    RWOutputTexture[DispatchThreadId] = TexelColor;
+    OutColor = InputTexture.Sample(InputSampler, InUVAndScreenPos.xy);
+    OutColor.rgb = pow(OutColor.rgb, abs(GammaPower));
 }
diff --git a/Engine/Plugins/Experimental/Compositing/HoldoutComposite/Shaders/Private/HoldoutCompositeSSRInput.usf b/Engine/Plugins/Experimental/Compositing/HoldoutComposite/Shaders/Private/HoldoutCompositeSSRInput.usf
--- a/Engine/Plugins/Experimental/Compositing/HoldoutComposite/Shaders/Private/HoldoutCompositeSSRInput.usf
+++ b/Engine/Plugins/Experimental/Compositing/HoldoutComposite/Shaders/Private/HoldoutCompositeSSRInput.usf
@@ -5,14 +5,14 @@
 #include "/Engine/Private/LensDistortion.ush"
 
 SCREEN_PASS_TEXTURE_VIEWPORT(Input)
-SCREEN_PASS_TEXTURE_VIEWPORT(Custom)
+SCREEN_PASS_TEXTURE_VIEWPORT(Composite)
 SCREEN_PASS_TEXTURE_VIEWPORT(Output)
 
 Texture2D InputTexture;
 SamplerState InputSampler;
 
-Texture2D CustomTexture;
-SamplerState CustomSampler;
+Texture2D CompositeTexture;
+SamplerState CompositeSampler;
 
 Texture2D<float2> DistortingDisplacementTexture;
 SamplerState DistortingDisplacementSampler;
@@ -33,10 +33,10 @@ void MainPS(
     const float2 InputUV = (DistortedViewportUV * Input_ViewportSize + Input_ViewportMin) * Input_ExtentInverse;
     float4 SceneColor = Texture2DSample(InputTexture, InputSampler, InputUV);
 
-    const float2 CustomUV = (ViewportUV * Custom_ViewportSize + Custom_ViewportMin) * Custom_ExtentInverse;
-    float4 CustomColor = Texture2DSample(CustomTexture, CustomSampler, CustomUV);
+    const float2 CompositeUV = (ViewportUV * Composite_ViewportSize + Composite_ViewportMin) * Composite_ExtentInverse;
+    float4 CompositeColor = Texture2DSample(CompositeTexture, CompositeSampler, CompositeUV);
 
-    CustomColor.rgb *= LastGlobalExposure;
-    OutColor.rgb = CustomColor.rgb * SceneColor.a + SceneColor.rgb;
-    OutColor.a = CustomColor.a * SceneColor.a;
+    CompositeColor.rgb *= LastGlobalExposure;
+    OutColor.rgb = CompositeColor.rgb * SceneColor.a + SceneColor.rgb;
+    OutColor.a = CompositeColor.a * SceneColor.a;
 }
diff --git a/Engine/Plugins/Experimental/Compositing/HoldoutComposite/Source/HoldoutComposite/Private/HoldoutComposite.cpp b/Engine/Plugins/Experimental/Compositing/HoldoutComposite/Source/HoldoutComposite/Private/HoldoutComposite.cpp
--- a/Engine/Plugins/Experimental/Compositing/HoldoutComposite/Source/HoldoutComposite/Private/HoldoutComposite.cpp
+++ b/Engine/Plugins/Experimental/Compositing/HoldoutComposite/Source/HoldoutComposite/Private/HoldoutComposite.cpp
@@ -8,6 +8,7 @@
 #include "Components/PrimitiveComponent.h"
 #include "Engine/Engine.h"
 #include "Engine/World.h"
+#include "Engine/Texture.h"
 #include "UObject/UObjectIterator.h"
 
 #include UE_INLINE_GENERATED_CPP_BY_NAME(HoldoutComposite)
@@ -149,6 +150,7 @@ void AHoldoutComposite::RegisterCompositeImpl()
 			Subsystem->OverrideCompositeSettings(SettingsOverride);
 		}
 		Subsystem->RegisterPrimitives(GetPrimitives(CompositeMeshActors));
+		Subsystem->SetCompositeTextureOverride(bOverrideCompositeTexture ? CompositeTexture : nullptr);
 		Subsystem->SetCompositePostProcessMaterials(CompositePostProcessMaterials);
 	}
 }
@@ -162,6 +164,7 @@ void AHoldoutComposite::UnregisterCompositeImpl()
 	{
 		Subsystem->ResetCompositePostProcessMaterials();
 		Subsystem->UnregisterPrimitives(GetPrimitives(CompositeMeshActors));
+		Subsystem->SetCompositeTextureOverride(nullptr);
 		Subsystem->ResetCompositeSettings();
 	}
 }
diff --git a/Engine/Plugins/Experimental/Compositing/HoldoutComposite/Source/HoldoutComposite/Private/HoldoutCompositeSceneViewExtension.cpp b/Engine/Plugins/Experimental/Compositing/HoldoutComposite/Source/HoldoutComposite/Private/HoldoutCompositeSceneViewExtension.cpp
--- a/Engine/Plugins/Experimental/Compositing/HoldoutComposite/Source/HoldoutComposite/Private/HoldoutCompositeSceneViewExtension.cpp
+++ b/Engine/Plugins/Experimental/Compositing/HoldoutComposite/Source/HoldoutComposite/Private/HoldoutCompositeSceneViewExtension.cpp
@@ -9,6 +9,7 @@
 #include "Containers/Map.h"
 #include "Containers/Set.h"
 #include "EngineUtils.h"
+#include "Engine/Texture.h"
 #include "FXRenderingUtils.h"
 #include "HDRHelper.h"
 #include "Materials/Material.h"
@@ -22,6 +23,7 @@
 #include "SceneView.h"
 #include "ScreenPass.h"
 #include "SystemTextures.h"
+#include "TextureResource.h"
 
 DECLARE_GPU_STAT_NAMED(FHoldoutCompositeDilate, TEXT("HoldoutComposite.Dilate"));
 DECLARE_GPU_STAT_NAMED(FHoldoutCompositeGamma, TEXT("HoldoutComposite.Gamma"));
@@ -38,6 +40,33 @@ namespace HoldoutComposite
 	};
 
 	static constexpr float DefaultGamma = 2.2f;
+
+	// Return the RDG texture from the the override texture uobject
+	FRDGTextureRef GetTextureOverride_RenderThread(FRDGBuilder& GraphBuilder, const TWeakObjectPtr<UTexture> InTexture)
+	{
+		constexpr bool bEvenIfPendingKill = false;
+		constexpr bool bThreadSafeTest = true;
+		if (InTexture.IsValid(bEvenIfPendingKill, bThreadSafeTest))
+		{
+			if (FTextureResource* TextureResource = InTexture.GetEvenIfUnreachable()->GetResource())
+			{
+				const FTextureRHIRef& TextureRHI = TextureResource->GetTextureRHI();
+				if (TextureRHI.IsValid())
+				{
+					FRDGTextureRef TextureRDG = GraphBuilder.FindExternalTexture(TextureRHI.GetReference());
+
+					if (TextureRDG == nullptr)
+					{
+						TextureRDG = GraphBuilder.RegisterExternalTexture(CreateRenderTarget(TextureRHI.GetReference(), TEXT("CompositeTextureOverride")));
+					}
+
+					return TextureRDG;
+				}
+			}
+		}
+
+		return nullptr;
+	}
 }
 
 class FHoldoutCompositeDilateShader : public FGlobalShader
@@ -69,30 +98,22 @@ class FHoldoutCompositeGammaShader : public FGlobalShader
 
 	BEGIN_SHADER_PARAMETER_STRUCT(FParameters, )
 		SHADER_PARAMETER_RDG_TEXTURE(Texture2D, InputTexture)
-		SHADER_PARAMETER_RDG_TEXTURE_UAV(RWTexture2D<float4>, RWOutputTexture)
+		SHADER_PARAMETER_SAMPLER(SamplerState, InputSampler)
 		SHADER_PARAMETER(FIntPoint, Dimensions)
 		SHADER_PARAMETER(float, GammaPower)
+		RENDER_TARGET_BINDING_SLOTS()
 	END_SHADER_PARAMETER_STRUCT()
-
-	static inline void ModifyCompilationEnvironment(const FGlobalShaderPermutationParameters& Parameters, FShaderCompilerEnvironment& OutEnvironment)
-	{
-		OutEnvironment.SetDefine(TEXT("TILE_SIZE"), ThreadGroupSize);
-
-		FGlobalShader::ModifyCompilationEnvironment(Parameters, OutEnvironment);
-	}
-
-	static const uint32 ThreadGroupSize = 16;
 };
-IMPLEMENT_GLOBAL_SHADER(FHoldoutCompositeGammaShader, "/Plugin/HoldoutComposite/Private/HoldoutCompositeGamma.usf", "MainCS", SF_Compute);
+IMPLEMENT_GLOBAL_SHADER(FHoldoutCompositeGammaShader, "/Plugin/HoldoutComposite/Private/HoldoutCompositeGamma.usf", "MainPS", SF_Pixel);
 
 BEGIN_SHADER_PARAMETER_STRUCT(FHoldoutCompositeCommonParameters, )
 	SHADER_PARAMETER_STRUCT(FScreenPassTextureViewportParameters, Input)
-	SHADER_PARAMETER_STRUCT(FScreenPassTextureViewportParameters, Custom)
+	SHADER_PARAMETER_STRUCT(FScreenPassTextureViewportParameters, Composite)
 	SHADER_PARAMETER_STRUCT(FScreenPassTextureViewportParameters, Output)
 	SHADER_PARAMETER_RDG_TEXTURE(Texture2D, InputTexture)
 	SHADER_PARAMETER_SAMPLER(SamplerState, InputSampler)
-	SHADER_PARAMETER_RDG_TEXTURE(Texture2D, CustomTexture)
-	SHADER_PARAMETER_SAMPLER(SamplerState, CustomSampler)
+	SHADER_PARAMETER_RDG_TEXTURE(Texture2D, CompositeTexture)
+	SHADER_PARAMETER_SAMPLER(SamplerState, CompositeSampler)
 	SHADER_PARAMETER_RDG_TEXTURE(Texture2D, DistortingDisplacementTexture)
 	SHADER_PARAMETER_SAMPLER(SamplerState, DistortingDisplacementSampler)
 	SHADER_PARAMETER_RDG_TEXTURE(Texture2D, UndistortingDisplacementTexture)
@@ -127,6 +148,7 @@ class FHoldoutCompositeFinalShader : public FGlobalShader
 		SHADER_PARAMETER(FUint32Vector, Encodings)
 		SHADER_PARAMETER(FVector3f, DisplayGamma)
 		SHADER_PARAMETER(uint32, bApplyTSRLensDistortion)
+		SHADER_PARAMETER(uint32, bIsCustomRenderPassTexture)
 		RENDER_TARGET_BINDING_SLOTS()
 	END_SHADER_PARAMETER_STRUCT()
 };
@@ -158,28 +180,30 @@ namespace HoldoutComposite
 		);
 	}
 
-	void AddGammaPass(FRDGBuilder& GraphBuilder, FRDGTextureRef InOutTexture, ERHIFeatureLevel::Type FeatureLevel, float GammaPower)
+	void AddGammaPass(FRDGBuilder& GraphBuilder, const FScreenPassViewInfo ViewInfo, FRDGTextureRef InOutTexture, float GammaPower)
 	{
 		RDG_EVENT_SCOPE_STAT(GraphBuilder, FHoldoutCompositeGamma, "HoldoutComposite.Gamma");
 		RDG_GPU_STAT_SCOPE(GraphBuilder, FHoldoutCompositeGamma);
 
-		FGlobalShaderMap* GlobalShaderMap = GetGlobalShaderMap(FeatureLevel);
-		const FIntPoint TextureSize = InOutTexture->Desc.Extent;
+		FGlobalShaderMap* GlobalShaderMap = GetGlobalShaderMap(ViewInfo.FeatureLevel);
+		FScreenPassRenderTarget Output = FScreenPassRenderTarget(InOutTexture, ERenderTargetLoadAction::ENoAction);
+		FScreenPassTextureViewport Viewport = FScreenPassTextureViewport(InOutTexture);
 
-		FHoldoutCompositeGammaShader::FParameters* PassParameters = GraphBuilder.AllocParameters<FHoldoutCompositeGammaShader::FParameters>();
-		PassParameters->InputTexture = InOutTexture;
-		PassParameters->RWOutputTexture = GraphBuilder.CreateUAV(InOutTexture);
-		PassParameters->Dimensions = TextureSize;
-		PassParameters->GammaPower = GammaPower;
+		FHoldoutCompositeGammaShader::FParameters* Parameters = GraphBuilder.AllocParameters<FHoldoutCompositeGammaShader::FParameters>();
+		Parameters->InputTexture = InOutTexture;
+		Parameters->InputSampler = TStaticSamplerState<>::GetRHI();
+		Parameters->GammaPower = GammaPower;
+		Parameters->RenderTargets[0] = Output.GetRenderTargetBinding();
 
-		TShaderMapRef<FHoldoutCompositeGammaShader> ComputeShader(GlobalShaderMap);
-		FComputeShaderUtils::AddPass(
+		TShaderMapRef<FHoldoutCompositeGammaShader> PixelShader(GlobalShaderMap);
+		AddDrawScreenPass(
 			GraphBuilder,
-			RDG_EVENT_NAME("HoldoutComposite.Gamma (%dx%d)", TextureSize.X, TextureSize.Y),
-			GSupportsEfficientAsyncCompute ? ERDGPassFlags::AsyncCompute : ERDGPassFlags::Compute,
-			ComputeShader,
-			PassParameters,
-			FComputeShaderUtils::GetGroupCount(TextureSize, FHoldoutCompositeGammaShader::ThreadGroupSize)
+			RDG_EVENT_NAME("HoldoutComposite.Gamma (%dx%d)", InOutTexture->Desc.Extent.X, InOutTexture->Desc.Extent.Y),
+			ViewInfo,
+			Viewport,
+			Viewport,
+			PixelShader,
+			Parameters
 		);
 	}
 }
@@ -307,10 +331,15 @@ void FHoldoutCompositeSceneViewExtension::ResetCompositePostProcessMaterials_Gam
 	CompositeProcessMaterials.Reset();
 }
 
+void FHoldoutCompositeSceneViewExtension::SetCompositeTextureOverride_GameThread(UTexture* InTextureOverride)
+{
+	CompositeTextureOverride = InTextureOverride;
+}
+
 bool FHoldoutCompositeSceneViewExtension::IsActiveThisFrame_Internal(const FSceneViewExtensionContext& Context) const
 {
 	bool bIsActive = FWorldSceneViewExtension::IsActiveThisFrame_Internal(Context);
-	bIsActive &= !CompositePrimitives.IsEmpty();
+	bIsActive &= !CompositePrimitives.IsEmpty() || CompositeTextureOverride.IsValid();
 	bIsActive &= !IsHDREnabled();
 
 	return bIsActive;
@@ -331,6 +360,12 @@ void FHoldoutCompositeSceneViewExtension::SetupViewFamily(FSceneViewFamily& InVi
 
 void FHoldoutCompositeSceneViewExtension::SetupView(FSceneViewFamily& InViewFamily, FSceneView& InView)
 {
+	// Skip custom render pass since we'll use the composite texture override directly
+	if (CompositeTextureOverride.IsValid())
+	{
+		return;
+	}
+
 	const TWeakObjectPtr<UWorld> WorldPtr = GetWorld();
 	check(WorldPtr.IsValid());
 
@@ -390,7 +425,8 @@ void FHoldoutCompositeSceneViewExtension::SetupView(FSceneViewFamily& InViewFami
 	PassInput.EngineShowFlags.SetTranslucency(true);
 	PassInput.EngineShowFlags.SetUnlitViewmode(GetSettings().bEnableUnlitViewmode);
 	PassInput.EngineShowFlags.SetAllowPrimitiveAlphaHoldout(false);
-
+	PassInput.bOverridesPostVolumeUserFlags = GetSettings().bOverridesViewUserFlags;
+	PassInput.PostVolumeUserFlags = GetSettings().CompositeViewUserFlags;
 	// Note: Incoming view location is invalid for scene captures
 	PassInput.ViewLocation = InView.bIsSceneCapture ? InView.ViewMatrices.GetViewOrigin() : InView.ViewLocation;
 	PassInput.ViewRotationMatrix = InView.ViewMatrices.GetViewMatrix().RemoveTranslation();
@@ -408,7 +444,9 @@ void FHoldoutCompositeSceneViewExtension::SetupView(FSceneViewFamily& InViewFami
 
 void FHoldoutCompositeSceneViewExtension::SubscribeToPostProcessingPass(EPostProcessingPass PassId, const FSceneView& InView, FAfterPassCallbackDelegateArray& InOutPassCallbacks, bool bIsPassEnabled)
 {
-	if (!CustomRenderTargetPerView_RenderThread.Contains(InView.GetViewKey()))
+	constexpr bool bEvenIfPendingKill = false;
+	constexpr bool bThreadsafeTest = true;
+	if (!CustomRenderTargetPerView_RenderThread.Contains(InView.GetViewKey()) && !CompositeTextureOverride.IsValid(bEvenIfPendingKill, bThreadsafeTest))
 	{
 		// Early-out to avoid needless work in the post processing callback(s).
 		return;
@@ -454,8 +492,13 @@ const FHoldoutCompositeSettings& FHoldoutCompositeSceneViewExtension::GetSetting
 	return CompositeSettings;
 }
 
-FRDGTextureRef FHoldoutCompositeSceneViewExtension::GetCustomRenderPassTexture(FRDGBuilder& GraphBuilder, const FSceneView& InView) const
+FRDGTextureRef FHoldoutCompositeSceneViewExtension::GetCompositeTexture(FRDGBuilder& GraphBuilder, const FSceneView& InView) const
 {
+	if (FRDGTextureRef TextureOverrideRDG = HoldoutComposite::GetTextureOverride_RenderThread(GraphBuilder, CompositeTextureOverride))
+	{
+		return TextureOverrideRDG;
+	}
+
 	FRDGTextureRef CustomRenderPassTexture = GSystemTextures.GetBlackAlphaOneDummy(GraphBuilder);
 
 	const TRefCountPtr<IPooledRenderTarget>* CustomRenderPassRenderTargetPtr = CustomRenderTargetPerView_RenderThread.Find(InView.GetViewKey());
@@ -467,19 +510,19 @@ FRDGTextureRef FHoldoutCompositeSceneViewExtension::GetCustomRenderPassTexture(F
 	return CustomRenderPassTexture;
 }
 
-FHoldoutCompositeCommonParameters FHoldoutCompositeSceneViewExtension::BuildCommonCompositeParameters(FRDGBuilder& GraphBuilder, const FSceneView& InView, const FScreenPassTexture& SceneColor, const FScreenPassTexture& CustomRenderPassTexture, const FScreenPassRenderTarget& Output, bool bIsSceneColorUndistorted)
+FHoldoutCompositeCommonParameters FHoldoutCompositeSceneViewExtension::BuildCommonCompositeParameters(FRDGBuilder& GraphBuilder, const FSceneView& InView, const FScreenPassTexture& SceneColor, const FScreenPassTexture& CompositeTexture, const FScreenPassRenderTarget& Output, bool bIsSceneColorUndistorted)
 {
 	const FLensDistortionLUT& LensDistortionLUT = LensDistortion::GetLUTUnsafe(InView);
 	const bool bLensDistortionInTSR = (LensDistortion::GetPassLocationUnsafe(InView) == LensDistortion::EPassLocation::TSR);
 
 	FHoldoutCompositeCommonParameters CommonParameters;
 	CommonParameters.Input = GetScreenPassTextureViewportParameters(FScreenPassTextureViewport(SceneColor));
-	CommonParameters.Custom = GetScreenPassTextureViewportParameters(FScreenPassTextureViewport(CustomRenderPassTexture));
+	CommonParameters.Composite = GetScreenPassTextureViewportParameters(FScreenPassTextureViewport(CompositeTexture));
 	CommonParameters.Output = GetScreenPassTextureViewportParameters(FScreenPassTextureViewport(Output));
 	CommonParameters.InputTexture = SceneColor.Texture;
 	CommonParameters.InputSampler = bIsSceneColorUndistorted ? TStaticSamplerState<SF_Bilinear, AM_Mirror, AM_Mirror>::GetRHI() : TStaticSamplerState<SF_Point>::GetRHI();
-	CommonParameters.CustomTexture = CustomRenderPassTexture.Texture;
-	CommonParameters.CustomSampler = TStaticSamplerState<SF_Bilinear>::GetRHI();
+	CommonParameters.CompositeTexture = CompositeTexture.Texture;
+	CommonParameters.CompositeSampler = TStaticSamplerState<SF_Bilinear>::GetRHI();
 	CommonParameters.UndistortingDisplacementTexture = GSystemTextures.GetBlackDummy(GraphBuilder);
 	CommonParameters.UndistortingDisplacementSampler = TStaticSamplerState<SF_Bilinear>::GetRHI();
 	CommonParameters.DistortingDisplacementTexture = GSystemTextures.GetBlackDummy(GraphBuilder);
@@ -502,12 +545,12 @@ FScreenPassTexture FHoldoutCompositeSceneViewExtension::PostProcessPassSSRInput_
 
 	FScreenPassTexture SceneColor = FScreenPassTexture::CopyFromSlice(GraphBuilder, Inputs.GetInput(EPostProcessMaterialInput::SceneColor));
 	check(SceneColor.IsValid());
-	FRDGTextureRef CustomRenderPassTexture = GetCustomRenderPassTexture(GraphBuilder, InView);
+	FScreenPassTexture CompositeTexture{ GetCompositeTexture(GraphBuilder, InView) };
 	FScreenPassRenderTarget Output = FScreenPassRenderTarget::CreateFromInput(GraphBuilder, SceneColor, InView.GetOverwriteLoadAction(), TEXT("HoldoutCompositeSSRInputRT"));
 
 	FHoldoutCompositeSSRInputShader::FParameters* PassParameters = GraphBuilder.AllocParameters<FHoldoutCompositeSSRInputShader::FParameters>();
 	constexpr bool bIsSceneColorUndistorted = true;
-	PassParameters->Common = BuildCommonCompositeParameters(GraphBuilder, InView, SceneColor, FScreenPassTexture(CustomRenderPassTexture), Output, bIsSceneColorUndistorted);
+	PassParameters->Common = BuildCommonCompositeParameters(GraphBuilder, InView, SceneColor, CompositeTexture, Output, bIsSceneColorUndistorted);
 	PassParameters->LastGlobalExposure = GetSettings().bFollowSceneExposure ? InView.GetLastEyeAdaptationExposure() : 1.0f;
 	PassParameters->RenderTargets[0] = Output.GetRenderTargetBinding();
 
@@ -534,7 +577,7 @@ FScreenPassTexture FHoldoutCompositeSceneViewExtension::PostProcessCompositePass
 	FScreenPassTexture SceneColor = FScreenPassTexture::CopyFromSlice(GraphBuilder, Inputs.GetInput(EPostProcessMaterialInput::SceneColor));
 	check(SceneColor.IsValid());
 
-	FScreenPassTexture CustomRenderPassTexture = FScreenPassTexture(GetCustomRenderPassTexture(GraphBuilder, InView));
+	FScreenPassTexture CompositeTexture{ GetCompositeTexture(GraphBuilder, InView) };
 
 	{
 		FScopeLock Lock(&MaterialsCS);
@@ -544,15 +587,19 @@ FScreenPassTexture FHoldoutCompositeSceneViewExtension::PostProcessCompositePass
 			if (Material.IsValid())
 			{
 				FPostProcessMaterialInputs PostProcessMaterialInputs;
-				PostProcessMaterialInputs.SetInput(GraphBuilder, EPostProcessMaterialInput::SceneColor, CustomRenderPassTexture);
+				PostProcessMaterialInputs.SetInput(GraphBuilder, EPostProcessMaterialInput::SceneColor, CompositeTexture);
 				PostProcessMaterialInputs.SceneTextures = Inputs.SceneTextures; //Note: these are invalid as they come from the main render, but validation-required.
 
 				// Create our own output override to maintain UAV access
-				FRDGTextureDesc OutputDesc = CustomRenderPassTexture.Texture->Desc;
+				FRDGTextureDesc OutputDesc = CompositeTexture.Texture->Desc;
+				OutputDesc.Reset();
+				OutputDesc.ClearValue = FClearValueBinding(FLinearColor::Black);
+				OutputDesc.Format = SceneColor.Texture->Desc.Format;
+				OutputDesc.Extent = CompositeTexture.ViewRect.Size();
 				OutputDesc.Flags |= ETextureCreateFlags::RenderTargetable;
-				PostProcessMaterialInputs.OverrideOutput = FScreenPassRenderTarget(GraphBuilder.CreateTexture(OutputDesc, TEXT("PostProcessMaterial")), CustomRenderPassTexture.ViewRect, InView.GetOverwriteLoadAction());
+				PostProcessMaterialInputs.OverrideOutput = FScreenPassRenderTarget(GraphBuilder.CreateTexture(OutputDesc, TEXT("PostProcessMaterial")), CompositeTexture.ViewRect, InView.GetOverwriteLoadAction());
 
-				CustomRenderPassTexture = AddPostProcessMaterialPass(GraphBuilder, InView, PostProcessMaterialInputs, Material.Get());
+				CompositeTexture = AddPostProcessMaterialPass(GraphBuilder, InView, PostProcessMaterialInputs, Material.Get());
 			}
 		}
 	}
@@ -560,14 +607,14 @@ FScreenPassTexture FHoldoutCompositeSceneViewExtension::PostProcessCompositePass
 	if (GetSettings().bApplyFXAA)
 	{
 		// We gamma-encode the result so that FXAA can operate on non-linear colors
-		HoldoutComposite::AddGammaPass(GraphBuilder, CustomRenderPassTexture.Texture, InView.GetFeatureLevel(), 1.0f / HoldoutComposite::DefaultGamma);
+		HoldoutComposite::AddGammaPass(GraphBuilder, InView, CompositeTexture.Texture, 1.0f / HoldoutComposite::DefaultGamma);
 
 		static TConsoleVariableData<int32>* CVarFXAAQuality = IConsoleManager::Get().FindTConsoleVariableDataInt(TEXT("r.FXAA.Quality"));
 
 		FFXAAInputs PassInputs{};
-		PassInputs.SceneColor = CustomRenderPassTexture;
+		PassInputs.SceneColor = CompositeTexture;
 		PassInputs.Quality = static_cast<EFXAAQuality>(FMath::Clamp(CVarFXAAQuality->GetValueOnRenderThread(), 0, 5));
-		CustomRenderPassTexture = AddFXAAPass(GraphBuilder, InView, PassInputs);
+		CompositeTexture = AddFXAAPass(GraphBuilder, InView, PassInputs);
 	}
 
 	FGlobalShaderMap* GlobalShaderMap = GetGlobalShaderMap(InView.GetFeatureLevel());
@@ -597,10 +644,13 @@ FScreenPassTexture FHoldoutCompositeSceneViewExtension::PostProcessCompositePass
 	{
 		FHoldoutCompositeFinalShader::FPermutationDomain PermutationVector;
 		PermutationVector.Set<FHoldoutCompositeFinalShader::FUseGlobalExposure>(GetSettings().bFollowSceneExposure);
+		constexpr bool bEvenIfPendingKill = false;
+		constexpr bool bThreadsafeTest = true;
+		const bool bActiveTextureOverride = CompositeTextureOverride.IsValid(bEvenIfPendingKill, bThreadsafeTest);
 
 		FHoldoutCompositeFinalShader::FParameters* PassParameters = GraphBuilder.AllocParameters<FHoldoutCompositeFinalShader::FParameters>();
 		PassParameters->View = InView.ViewUniformBuffer;
-		PassParameters->Common = BuildCommonCompositeParameters(GraphBuilder, InView, SceneColor, CustomRenderPassTexture, Output);
+		PassParameters->Common = BuildCommonCompositeParameters(GraphBuilder, InView, SceneColor, CompositeTexture, Output);
 		if (GetSettings().bFollowSceneExposure)
 		{
 			FRDGBufferRef EyeAdaptationBuffer = GraphBuilder.RegisterExternalBuffer(InView.GetEyeAdaptationBuffer(), ERDGBufferFlags::MultiFrame);
@@ -609,6 +659,7 @@ FScreenPassTexture FHoldoutCompositeSceneViewExtension::PostProcessCompositePass
 		PassParameters->Encodings = Encodings;
 		PassParameters->DisplayGamma = FVector3f(Family->RenderTarget->GetDisplayGamma(), 1.0f / Family->RenderTarget->GetDisplayGamma(), HoldoutComposite::DefaultGamma);
 		PassParameters->bApplyTSRLensDistortion = GetSettings().bApplyTSRLensDistortion;
+		PassParameters->bIsCustomRenderPassTexture = !bActiveTextureOverride;
 		PassParameters->RenderTargets[0] = Output.GetRenderTargetBinding();
 
 		TShaderMapRef<FHoldoutCompositeFinalShader> PixelShader(GlobalShaderMap, PermutationVector);
diff --git a/Engine/Plugins/Experimental/Compositing/HoldoutComposite/Source/HoldoutComposite/Private/HoldoutCompositeSceneViewExtension.h b/Engine/Plugins/Experimental/Compositing/HoldoutComposite/Source/HoldoutComposite/Private/HoldoutCompositeSceneViewExtension.h
--- a/Engine/Plugins/Experimental/Compositing/HoldoutComposite/Source/HoldoutComposite/Private/HoldoutCompositeSceneViewExtension.h
+++ b/Engine/Plugins/Experimental/Compositing/HoldoutComposite/Source/HoldoutComposite/Private/HoldoutCompositeSceneViewExtension.h
@@ -15,6 +15,7 @@ struct FScreenPassRenderTarget;
 struct FScreenPassTexture;
 struct FHoldoutCompositeSettings;
 class UMaterialInterface;
+class UTexture;
 
 class FHoldoutCompositeSceneViewExtension : public FWorldSceneViewExtension
 {
@@ -40,6 +41,9 @@ public:
 	/* Remove all post-process materials to be applied on the separate composited render. */
 	void ResetCompositePostProcessMaterials_GameThread();
 
+	/* Set a texture override to be used instead of the separate composited render. */
+	void SetCompositeTextureOverride_GameThread(UTexture* InTextureOverride);
+
 	/* Called by the custom render pass to store its view render target for this frame. */
 	template<typename T>
 	void CollectCustomRenderTarget(uint32 InViewId, T InRenderTarget)
@@ -69,22 +73,28 @@ protected:
 
 private:
 
-	FRDGTextureRef GetCustomRenderPassTexture(FRDGBuilder& GraphBuilder, const FSceneView& InView) const;
-	class FHoldoutCompositeCommonParameters BuildCommonCompositeParameters(FRDGBuilder& GraphBuilder, const FSceneView& InView, const FScreenPassTexture& SceneColor, const FScreenPassTexture& CustomRenderPassTexture, const FScreenPassRenderTarget& Output, bool bIsSceneColorUndistorted = false);
+	/* Returns the custom render pass render target, or the composite texture override if present.*/
+	FRDGTextureRef GetCompositeTexture(FRDGBuilder& GraphBuilder, const FSceneView& InView) const;
+
+	/* Build the common composite pass shader parameters. */
+	class FHoldoutCompositeCommonParameters BuildCommonCompositeParameters(FRDGBuilder& GraphBuilder, const FSceneView& InView, const FScreenPassTexture& SceneColor, const FScreenPassTexture& CompositeTexture, const FScreenPassRenderTarget& Output, bool bIsSceneColorUndistorted = false);
 
 	FScreenPassTexture PostProcessPassSSRInput_RenderThread(FRDGBuilder& GraphBuilder, const FSceneView& InView, const FPostProcessMaterialInputs& Inputs);
 	FScreenPassTexture PostProcessCompositePass_RenderThread(FRDGBuilder& GraphBuilder, const FSceneView& InView, const FPostProcessMaterialInputs& Inputs);
 
-	// Collection of primitives to render as a custom render pass and composite after post-processing.
+	/* Collection of primitives to render as a custom render pass and composite after post-processing. .*/
 	TSet<TWeakObjectPtr<UPrimitiveComponent>> CompositePrimitives;
 
-	// Custom render pass render targets for each active view
+	/* Custom render pass render targets for each active view .*/
 	TMap<uint32, TRefCountPtr<IPooledRenderTarget>> CustomRenderTargetPerView_RenderThread;
 
-	// Composite settings and its render-thread copy
+	/* Composite settings and its render-thread copy */
 	FHoldoutCompositeSettings CompositeSettings;
 	FHoldoutCompositeSettings CompositeSettings_RenderThread;
 
+	/* Composite texture override use directly instead of the separate composited render. */
+	TWeakObjectPtr<UTexture> CompositeTextureOverride;
+
 	/** Post-process material(s) to be applied on the separate composited render.*/
 	TArray<TWeakObjectPtr<UMaterialInterface>> CompositeProcessMaterials;
 
diff --git a/Engine/Plugins/Experimental/Compositing/HoldoutComposite/Source/HoldoutComposite/Private/HoldoutCompositeSettings.cpp b/Engine/Plugins/Experimental/Compositing/HoldoutComposite/Source/HoldoutComposite/Private/HoldoutCompositeSettings.cpp
--- a/Engine/Plugins/Experimental/Compositing/HoldoutComposite/Source/HoldoutComposite/Private/HoldoutCompositeSettings.cpp
+++ b/Engine/Plugins/Experimental/Compositing/HoldoutComposite/Source/HoldoutComposite/Private/HoldoutCompositeSettings.cpp
@@ -15,6 +15,8 @@ FHoldoutCompositeSettings::FHoldoutCompositeSettings()
 	, bApplyFXAA(false)
 	, CompositePassLocation(EHoldoutCompositePassLocation::AfterTonemap)
 	, CompositePassPriority(100 + 1 /* OPENCOLORIO_SCENE_VIEW_EXTENSION_PRIORITY + 1 */)
+	, bOverridesViewUserFlags(false)
+	, CompositeViewUserFlags(0)
 { }
 
 UHoldoutCompositePluginSettings::UHoldoutCompositePluginSettings()
diff --git a/Engine/Plugins/Experimental/Compositing/HoldoutComposite/Source/HoldoutComposite/Private/HoldoutCompositeSubsystem.cpp b/Engine/Plugins/Experimental/Compositing/HoldoutComposite/Source/HoldoutComposite/Private/HoldoutCompositeSubsystem.cpp
--- a/Engine/Plugins/Experimental/Compositing/HoldoutComposite/Source/HoldoutComposite/Private/HoldoutCompositeSubsystem.cpp
+++ b/Engine/Plugins/Experimental/Compositing/HoldoutComposite/Source/HoldoutComposite/Private/HoldoutCompositeSubsystem.cpp
@@ -128,6 +128,14 @@ void UHoldoutCompositeSubsystem::ResetCompositeSettings()
 	}
 }
 
+void UHoldoutCompositeSubsystem::SetCompositeTextureOverride(UTexture* InTextureOverride)
+{
+	if (HoldoutCompositeViewExtension.IsValid())
+	{
+		HoldoutCompositeViewExtension->SetCompositeTextureOverride_GameThread(InTextureOverride);
+	}
+}
+
 bool UHoldoutCompositeSubsystem::ValidateProjectSettings()
 {
 	URendererSettings* RendererSettings = GetMutableDefault<URendererSettings>();
diff --git a/Engine/Plugins/Experimental/Compositing/HoldoutComposite/Source/HoldoutComposite/Public/HoldoutComposite.h b/Engine/Plugins/Experimental/Compositing/HoldoutComposite/Source/HoldoutComposite/Public/HoldoutComposite.h
--- a/Engine/Plugins/Experimental/Compositing/HoldoutComposite/Source/HoldoutComposite/Public/HoldoutComposite.h
+++ b/Engine/Plugins/Experimental/Compositing/HoldoutComposite/Source/HoldoutComposite/Public/HoldoutComposite.h
@@ -10,6 +10,7 @@
 
 class UPrimitiveComponent;
 class UMaterialInterface;
+class UTexture;
 
 UCLASS()
 class HOLDOUTCOMPOSITE_API AHoldoutComposite : public AActor
@@ -47,7 +48,7 @@ private:
 
 public:
 
-	/* Actors whose primitive (mesh) components will be marked as primitive alpha holdout, rendered separately and composited back under scene color after post-processing. */
+	/* Actors whose primitive (mesh) components will be marked as primitive alpha holdout, rendered separately and composited with scene color after post-processing. */
 	UPROPERTY(EditAnywhere, BlueprintReadOnly, Category = HoldoutComposite)
 	TArray<TObjectPtr<AActor>> CompositeMeshActors;
 
@@ -55,7 +56,15 @@ public:
 	UPROPERTY(EditAnywhere, BlueprintReadOnly, Category = HoldoutComposite)
 	TArray<TObjectPtr<UMaterialInterface>> CompositePostProcessMaterials;
 
-	/** Post-process material(s) to be applied on the separate composited render.*/
+	/** Enable composite texture override.*/
+	UPROPERTY(EditAnywhere, BlueprintReadOnly, Category = HoldoutComposite)
+	bool bOverrideCompositeTexture = false;
+
+	/* Instead of rendering the holdout composite actors separately, we only use them as primitive alpha holdout masks and composite the texture directly with scene color. */
+	UPROPERTY(EditAnywhere, BlueprintReadOnly, Category = HoldoutComposite, meta = (EditCondition = bOverrideCompositeTexture))
+	TObjectPtr<UTexture> CompositeTexture;
+
+	/** Enable local settings override.*/
 	UPROPERTY(EditAnywhere, BlueprintReadOnly, Category = HoldoutComposite)
 	bool bOverrideDefaultSettings = false;
 
diff --git a/Engine/Plugins/Experimental/Compositing/HoldoutComposite/Source/HoldoutComposite/Public/HoldoutCompositeSettings.h b/Engine/Plugins/Experimental/Compositing/HoldoutComposite/Source/HoldoutComposite/Public/HoldoutCompositeSettings.h
--- a/Engine/Plugins/Experimental/Compositing/HoldoutComposite/Source/HoldoutComposite/Public/HoldoutCompositeSettings.h
+++ b/Engine/Plugins/Experimental/Compositing/HoldoutComposite/Source/HoldoutComposite/Public/HoldoutCompositeSettings.h
@@ -57,6 +57,14 @@ struct HOLDOUTCOMPOSITE_API FHoldoutCompositeSettings
 	/** Composite (scene view extension) pass priority, which defaults to before OpenColorIO. */
 	UPROPERTY(config, EditAnywhere, Category = HoldoutComposite, AdvancedDisplay)
 	int32 CompositePassPriority;
+
+	/** Override the composite view user flags with the value specified below. */
+	UPROPERTY(config, EditAnywhere, Category = HoldoutComposite, AdvancedDisplay)
+	bool bOverridesViewUserFlags;
+
+	/** User flag for the composite render, allowing per-view material overrides. Active if greater or equal to zero. */
+	UPROPERTY(config, EditAnywhere, Category = HoldoutComposite, AdvancedDisplay, meta = (EditCondition = "bOverridesViewUserFlags"))
+	int32 CompositeViewUserFlags;
 };
 
 /**
diff --git a/Engine/Plugins/Experimental/Compositing/HoldoutComposite/Source/HoldoutComposite/Public/HoldoutCompositeSubsystem.h b/Engine/Plugins/Experimental/Compositing/HoldoutComposite/Source/HoldoutComposite/Public/HoldoutCompositeSubsystem.h
--- a/Engine/Plugins/Experimental/Compositing/HoldoutComposite/Source/HoldoutComposite/Public/HoldoutCompositeSubsystem.h
+++ b/Engine/Plugins/Experimental/Compositing/HoldoutComposite/Source/HoldoutComposite/Public/HoldoutCompositeSubsystem.h
@@ -59,6 +59,10 @@ public:
 	UFUNCTION(BlueprintCallable, Category = "Holdout Composite")
 	HOLDOUTCOMPOSITE_API void ResetCompositeSettings();
 
+	/* Set a texture override to be used instead of the separate composited render. */
+	UFUNCTION(BlueprintCallable, Category = "Holdout Composite")
+	void SetCompositeTextureOverride(UTexture* InTextureOverride);
+
 private:
 
 	/* Returns true if the (renderer) project settings are correctly enabled for the composite to be active. */
