Holdout: Minor changes in HoldoutComposite plugin
- The `UHoldoutCompositeSubsystem` class is modified so that it can be inherited by a class from another module.
- Added new parameter `bLinearComposite` to `FHoldoutCompositeSettings`. To work with a composite in the sRGB color space.
- The `FHoldoutCompositeSceneViewExtension` has been modified to implement a new customization parameter
- Minor changes in 'HoldoutCompositeFinal.usf' shader.

diff --git a/Engine/Plugins/Experimental/Compositing/HoldoutComposite/Shaders/Private/HoldoutCompositeFinal.usf b/Engine/Plugins/Experimental/Compositing/HoldoutComposite/Shaders/Private/HoldoutCompositeFinal.usf
--- a/Engine/Plugins/Experimental/Compositing/HoldoutComposite/Shaders/Private/HoldoutCompositeFinal.usf
+++ b/Engine/Plugins/Experimental/Compositing/HoldoutComposite/Shaders/Private/HoldoutCompositeFinal.usf
@@ -58,11 +58,6 @@ void MainPS(
     const float2 CompositeUV = (CompositeViewportUV * Composite_ViewportSize + Composite_ViewportMin) * Composite_ExtentInverse;
     float4 CompositeColor = Texture2DSample(CompositeTexture, CompositeSampler, CompositeUV);
 
-#if USE_GLOBAL_EXPOSURE
-    const float GlobalExposure = EyeAdaptationBuffer[0].x;
-    CompositeColor.rgb *= GlobalExposure;
-#endif
-
     BRANCH
     if(Encodings.x == 1)
     {
@@ -74,10 +69,22 @@ void MainPS(
     }
 
     // Optionally linearize composite color
-    if(Encodings.z > 0)
+	BRANCH
+    if(Encodings.z == 1)
     {
         CompositeColor.rgb = pow(CompositeColor.rgb, DisplayGamma.z);
     }
+	else if(Encodings.z == 2)
+    {
+        CompositeColor.rgb = sRGBToLinear(CompositeColor.rgb);
+    }
+
+#if USE_GLOBAL_EXPOSURE
+    // Applying GlobalExposure in linear color space
+    const float GlobalExposure = EyeAdaptationBuffer[0].x;
+    CompositeColor.rgb *= GlobalExposure;
+#endif
+
 
 	// Custom render pass texture has inverted pre-multiplied alpha (like the renderer)
 	if (bIsCustomRenderPassTexture > 0)
diff --git a/Engine/Plugins/Experimental/Compositing/HoldoutComposite/Source/HoldoutComposite/Private/HoldoutCompositeSceneViewExtension.cpp b/Engine/Plugins/Experimental/Compositing/HoldoutComposite/Source/HoldoutComposite/Private/HoldoutCompositeSceneViewExtension.cpp
--- a/Engine/Plugins/Experimental/Compositing/HoldoutComposite/Source/HoldoutComposite/Private/HoldoutCompositeSceneViewExtension.cpp
+++ b/Engine/Plugins/Experimental/Compositing/HoldoutComposite/Source/HoldoutComposite/Private/HoldoutCompositeSceneViewExtension.cpp
@@ -273,8 +273,10 @@ void FHoldoutCompositeSceneViewExtension::RegisterPrimitives_GameThread(const TA
 		if (!CompositePrimitives.Contains(InPrimitiveComponent))
 		{
 			CompositePrimitives.Add(InPrimitiveComponent);
-			InPrimitiveComponent->SetHoldout(bInHoldoutState);
 		}
+
+		// The SetHoldout() function makes changes only if the HoldoutState value differs from the current value.
+		InPrimitiveComponent->SetHoldout(bInHoldoutState);
 	}
 }
 
@@ -292,8 +294,10 @@ void FHoldoutCompositeSceneViewExtension::UnregisterPrimitives_GameThread(const
 		if (CompositePrimitives.Contains(InPrimitiveComponent))
 		{
 			CompositePrimitives.Remove(InPrimitiveComponent);
-			InPrimitiveComponent->SetHoldout(bInHoldoutState);
 		}
+
+		// The SetHoldout() function makes changes only if the HoldoutState value differs from the current value.
+		InPrimitiveComponent->SetHoldout(bInHoldoutState);
 	}
 }
 
@@ -373,7 +377,9 @@ void FHoldoutCompositeSceneViewExtension::SetupView(FSceneViewFamily& InViewFami
 	for (const TWeakObjectPtr<UPrimitiveComponent>& PrimitivePtr : CompositePrimitives)
 	{
 		TStrongObjectPtr<UPrimitiveComponent> Primitive = PrimitivePtr.Pin();
-		if (Primitive.IsValid())
+		// Collect only those primitives that use the bHoldout flag
+		// The user can directly change this flag outside of this VE.
+		if (Primitive.IsValid() && Primitive->bHoldout)
 		{
 			const FPrimitiveComponentId PrimId = Primitive->GetPrimitiveSceneId();
 
@@ -606,8 +612,11 @@ FScreenPassTexture FHoldoutCompositeSceneViewExtension::PostProcessCompositePass
 
 	if (GetSettings().bApplyFXAA)
 	{
-		// We gamma-encode the result so that FXAA can operate on non-linear colors
-		HoldoutComposite::AddGammaPass(GraphBuilder, InView, CompositeTexture.Texture, 1.0f / HoldoutComposite::DefaultGamma);
+		if (GetSettings().bLinearComposite)
+		{
+			// We gamma-encode the result so that FXAA can operate on non-linear colors
+			HoldoutComposite::AddGammaPass(GraphBuilder, InView, CompositeTexture.Texture, 1.0f / HoldoutComposite::DefaultGamma);
+		}
 
 		static TConsoleVariableData<int32>* CVarFXAAQuality = IConsoleManager::Get().FindTConsoleVariableDataInt(TEXT("r.FXAA.Quality"));
 
@@ -631,8 +640,20 @@ FScreenPassTexture FHoldoutCompositeSceneViewExtension::PostProcessCompositePass
 		Encodings.X = static_cast<uint32>(ESceneColorSourceEncoding::sRGB);
 		Encodings.Y = static_cast<uint32>(ESceneColorSourceEncoding::sRGB);
 	}
-	// FXAA received gamma-encoded colors, so we must linearize before compositing.
-	Encodings.Z = static_cast<uint32>(GetSettings().bApplyFXAA);
+
+	// Encodings.Z controls the encoding of the composite texture
+	// A default value of 0 means that the composite is in linear color space
+	if (!GetSettings().bLinearComposite)
+	{
+		// The composite texture is in the sRGB color space and must be encoded in the linear color space
+		Encodings.Z = static_cast<uint32>(ESceneColorSourceEncoding::sRGB);
+	}
+	else if(GetSettings().bApplyFXAA)
+	{
+		// FXAA received gamma-encoded colors, so we must linearize before compositing.
+		Encodings.Z = static_cast<uint32>(ESceneColorSourceEncoding::Gamma);
+	}
+
 
 	FScreenPassRenderTarget Output = Inputs.OverrideOutput;
 	if (!Output.IsValid())
diff --git a/Engine/Plugins/Experimental/Compositing/HoldoutComposite/Source/HoldoutComposite/Private/HoldoutCompositeSettings.cpp b/Engine/Plugins/Experimental/Compositing/HoldoutComposite/Source/HoldoutComposite/Private/HoldoutCompositeSettings.cpp
--- a/Engine/Plugins/Experimental/Compositing/HoldoutComposite/Source/HoldoutComposite/Private/HoldoutCompositeSettings.cpp
+++ b/Engine/Plugins/Experimental/Compositing/HoldoutComposite/Source/HoldoutComposite/Private/HoldoutCompositeSettings.cpp
@@ -17,6 +17,7 @@ FHoldoutCompositeSettings::FHoldoutCompositeSettings()
 	, CompositePassPriority(100 + 1 /* OPENCOLORIO_SCENE_VIEW_EXTENSION_PRIORITY + 1 */)
 	, bOverridesViewUserFlags(false)
 	, CompositeViewUserFlags(0)
+	, bLinearComposite(true)
 { }
 
 UHoldoutCompositePluginSettings::UHoldoutCompositePluginSettings()
diff --git a/Engine/Plugins/Experimental/Compositing/HoldoutComposite/Source/HoldoutComposite/Private/HoldoutCompositeSubsystem.cpp b/Engine/Plugins/Experimental/Compositing/HoldoutComposite/Source/HoldoutComposite/Private/HoldoutCompositeSubsystem.cpp
--- a/Engine/Plugins/Experimental/Compositing/HoldoutComposite/Source/HoldoutComposite/Private/HoldoutCompositeSubsystem.cpp
+++ b/Engine/Plugins/Experimental/Compositing/HoldoutComposite/Source/HoldoutComposite/Private/HoldoutCompositeSubsystem.cpp
@@ -136,16 +136,33 @@ void UHoldoutCompositeSubsystem::SetCompositeTextureOverride(UTexture* InTexture
 	}
 }
 
-bool UHoldoutCompositeSubsystem::ValidateProjectSettings()
+bool UHoldoutCompositeSubsystem::EnableHoldoutSettingsNotification() const
+{
+	return true;
+}
+
+bool UHoldoutCompositeSubsystem::IsProjectSettingsValid()
 {
 	URendererSettings* RendererSettings = GetMutableDefault<URendererSettings>();
 	check(RendererSettings);
 
-	const bool bValidSettings = RendererSettings->bEnableAlphaChannelInPostProcessing && RendererSettings->bDeferredSupportPrimitiveAlphaHoldout;
+	const bool bValidSettings = 
+		RendererSettings->bEnableAlphaChannelInPostProcessing
+		&& RendererSettings->bDeferredSupportPrimitiveAlphaHoldout;
 
-	if (!bValidSettings)
+	return bValidSettings;
+}
+
+bool UHoldoutCompositeSubsystem::ValidateProjectSettings()
+{
+	const bool bValidSettings = IsProjectSettingsValid();
+
+	if (!bValidSettings && EnableHoldoutSettingsNotification())
 	{
 #if WITH_EDITOR
+		URendererSettings* RendererSettings = GetMutableDefault<URendererSettings>();
+		check(RendererSettings);
+
 		// We inform the user and offer them the option to activate project settings.
 		UE_CALL_ONCE([&]{ PrimitiveHoldoutSettingsNotification(RendererSettings); });
 #else
diff --git a/Engine/Plugins/Experimental/Compositing/HoldoutComposite/Source/HoldoutComposite/Public/HoldoutCompositeSettings.h b/Engine/Plugins/Experimental/Compositing/HoldoutComposite/Source/HoldoutComposite/Public/HoldoutCompositeSettings.h
--- a/Engine/Plugins/Experimental/Compositing/HoldoutComposite/Source/HoldoutComposite/Public/HoldoutCompositeSettings.h
+++ b/Engine/Plugins/Experimental/Compositing/HoldoutComposite/Source/HoldoutComposite/Public/HoldoutCompositeSettings.h
@@ -65,6 +65,10 @@ struct HOLDOUTCOMPOSITE_API FHoldoutCompositeSettings
 	/** User flag for the composite render, allowing per-view material overrides. Active if greater or equal to zero. */
 	UPROPERTY(config, EditAnywhere, Category = HoldoutComposite, AdvancedDisplay, meta = (EditCondition = "bOverridesViewUserFlags"))
 	int32 CompositeViewUserFlags;
+
+	/** If enabled, composite renders in linear colorspace. */
+	UPROPERTY(config, EditAnywhere, Category = HoldoutComposite, AdvancedDisplay)
+	bool bLinearComposite;
 };
 
 /**
diff --git a/Engine/Plugins/Experimental/Compositing/HoldoutComposite/Source/HoldoutComposite/Public/HoldoutCompositeSubsystem.h b/Engine/Plugins/Experimental/Compositing/HoldoutComposite/Source/HoldoutComposite/Public/HoldoutCompositeSubsystem.h
--- a/Engine/Plugins/Experimental/Compositing/HoldoutComposite/Source/HoldoutComposite/Public/HoldoutCompositeSubsystem.h
+++ b/Engine/Plugins/Experimental/Compositing/HoldoutComposite/Source/HoldoutComposite/Public/HoldoutCompositeSubsystem.h
@@ -17,7 +17,7 @@ class URendererSettings;
  * Composite subsytem used as an interface to the (private) scene view extension.
  */
 UCLASS(BlueprintType, Transient)
-class UHoldoutCompositeSubsystem : public UWorldSubsystem
+class HOLDOUTCOMPOSITE_API UHoldoutCompositeSubsystem : public UWorldSubsystem
 {
 	GENERATED_BODY()
 
@@ -31,37 +31,44 @@ public:
 
 	/* Register a single primitive for compositing. */
 	UFUNCTION(BlueprintCallable, Category = "Holdout Composite")
-	HOLDOUTCOMPOSITE_API void RegisterPrimitive(UPrimitiveComponent* InPrimitiveComponent, bool bInHoldoutState=true);
+	virtual void RegisterPrimitive(UPrimitiveComponent* InPrimitiveComponent, bool bInHoldoutState=true);
 
 	/* Register multiple primitives for compositing. */
-	HOLDOUTCOMPOSITE_API void RegisterPrimitives(const TArray<UPrimitiveComponent*>& InPrimitiveComponents, bool bInHoldoutState=true);
+	virtual void RegisterPrimitives(const TArray<UPrimitiveComponent*>& InPrimitiveComponents, bool bInHoldoutState=true);
 
 	/* Unregister a single primitive from compositing. */
 	UFUNCTION(BlueprintCallable, Category = "Holdout Composite")
-	HOLDOUTCOMPOSITE_API void UnregisterPrimitive(UPrimitiveComponent* InPrimitiveComponent, bool bInHoldoutState=false);
+	virtual void UnregisterPrimitive(UPrimitiveComponent* InPrimitiveComponent, bool bInHoldoutState=false);
 
 	/* Unregister multiple primitives from compositing. */
-	HOLDOUTCOMPOSITE_API void UnregisterPrimitives(const TArray<UPrimitiveComponent*>& InPrimitiveComponents, bool bInHoldoutState=false);
+	virtual void UnregisterPrimitives(const TArray<UPrimitiveComponent*>& InPrimitiveComponents, bool bInHoldoutState=false);
 
 	/* Set additional post-process materials to be applied on the separate composited render. */
 	UFUNCTION(BlueprintCallable, Category = "Holdout Composite")
-	HOLDOUTCOMPOSITE_API void SetCompositePostProcessMaterials(const TArray<UMaterialInterface*>& InMaterials);
+	virtual void SetCompositePostProcessMaterials(const TArray<UMaterialInterface*>& InMaterials);
 
 	/* Remove all post-process materials to be applied on the separate composited render. */
 	UFUNCTION(BlueprintCallable, Category = "Holdout Composite")
-	HOLDOUTCOMPOSITE_API void ResetCompositePostProcessMaterials();
+	virtual void ResetCompositePostProcessMaterials();
 
 	/* Override project default composite settings for the active holdout composite passes. */
 	UFUNCTION(BlueprintCallable, Category = "Holdout Composite")
-	HOLDOUTCOMPOSITE_API void OverrideCompositeSettings(const FHoldoutCompositeSettings& InSettings);
+	virtual void OverrideCompositeSettings(const FHoldoutCompositeSettings& InSettings);
 
 	/* Revert composite settings back to project defaults. */
 	UFUNCTION(BlueprintCallable, Category = "Holdout Composite")
-	HOLDOUTCOMPOSITE_API void ResetCompositeSettings();
+	virtual void ResetCompositeSettings();
 
 	/* Set a texture override to be used instead of the separate composited render. */
 	UFUNCTION(BlueprintCallable, Category = "Holdout Composite")
-	void SetCompositeTextureOverride(UTexture* InTextureOverride);
+	virtual void SetCompositeTextureOverride(UTexture* InTextureOverride);
+
+	/** True if the project settings are valid for the HoldoutComposite plugin to work. */
+	static bool IsProjectSettingsValid();
+
+protected:
+	/** Inherited classes can implement notification in a different way. */
+	virtual bool EnableHoldoutSettingsNotification() const;
 
 private:
 
