HoldoutComposite: Various API improvements and cleanup.

* Removed optional holdout state arguments, since they have little to no utility and may cause confusion.
* Added priority argument to override methods.
* Renamed encoding enum and mirrored definition in shader.
* Fixed plugin settings not updating live anymore, and lost values for disabled primitives.

diff --git a/Engine/Plugins/Experimental/Compositing/HoldoutComposite/Config/BaseHoldoutComposite.ini b/Engine/Plugins/Experimental/Compositing/HoldoutComposite/Config/BaseHoldoutComposite.ini
deleted file mode 100644
diff --git a/Engine/Plugins/Experimental/Compositing/HoldoutComposite/Shaders/Private/HoldoutCompositeFinal.usf b/Engine/Plugins/Experimental/Compositing/HoldoutComposite/Shaders/Private/HoldoutCompositeFinal.usf
--- a/Engine/Plugins/Experimental/Compositing/HoldoutComposite/Shaders/Private/HoldoutCompositeFinal.usf
+++ b/Engine/Plugins/Experimental/Compositing/HoldoutComposite/Shaders/Private/HoldoutCompositeFinal.usf
@@ -9,6 +9,10 @@ SCREEN_PASS_TEXTURE_VIEWPORT(Input)
 SCREEN_PASS_TEXTURE_VIEWPORT(Composite)
 SCREEN_PASS_TEXTURE_VIEWPORT(Output)
 
+#define ESourceEncoding_None  0
+#define ESourceEncoding_Gamma 1
+#define ESourceEncoding_sRGB  2
+
 Texture2D InputTexture;
 SamplerState InputSampler;
 
@@ -59,22 +63,22 @@ void MainPS(
     float4 CompositeColor = Texture2DSample(CompositeTexture, CompositeSampler, CompositeUV);
 
     BRANCH
-    if(Encodings.x == 1)
+    if(Encodings.x == ESourceEncoding_Gamma)
     {
         SceneColor.rgb = pow(SceneColor.rgb, DisplayGamma.x);
     }
-    else if(Encodings.x == 2)
+    else if(Encodings.x == ESourceEncoding_sRGB)
     {
         SceneColor.rgb = sRGBToLinear(SceneColor.rgb);
     }
 
     // Optionally linearize composite color
 	BRANCH
-    if(Encodings.z == 1)
+    if(Encodings.z == ESourceEncoding_Gamma)
     {
         CompositeColor.rgb = pow(CompositeColor.rgb, DisplayGamma.z);
     }
-	else if(Encodings.z == 2)
+	else if(Encodings.z == ESourceEncoding_sRGB)
     {
         CompositeColor.rgb = sRGBToLinear(CompositeColor.rgb);
     }
@@ -100,11 +104,11 @@ void MainPS(
     OutColor.a = CompositeColor.a * SceneColor.a;
 
     BRANCH
-    if(Encodings.y == 1)
+    if(Encodings.y == ESourceEncoding_Gamma)
     {
         OutColor.rgb = pow(OutColor.rgb, DisplayGamma.y);
     }
-    else if(Encodings.y == 2)
+    else if(Encodings.y == ESourceEncoding_sRGB)
     {
         OutColor.rgb = LinearToSrgb(OutColor.rgb);
     }
diff --git a/Engine/Plugins/Experimental/Compositing/HoldoutComposite/Source/HoldoutComposite/Private/HoldoutComposite.cpp b/Engine/Plugins/Experimental/Compositing/HoldoutComposite/Source/HoldoutComposite/Private/HoldoutComposite.cpp
--- a/Engine/Plugins/Experimental/Compositing/HoldoutComposite/Source/HoldoutComposite/Private/HoldoutComposite.cpp
+++ b/Engine/Plugins/Experimental/Compositing/HoldoutComposite/Source/HoldoutComposite/Private/HoldoutComposite.cpp
@@ -147,7 +147,9 @@ void AHoldoutComposite::RegisterCompositeImpl()
 	{
 		if (bOverrideDefaultSettings)
 		{
-			Subsystem->OverrideCompositeSettings(SettingsOverride);
+			// We leave the default holdout composite priority at 0, since only be one actor per level is allowed.
+			constexpr int32 DefaultPriority = 0;
+			Subsystem->OverrideCompositeSettings(SettingsOverride, DefaultPriority);
 		}
 		Subsystem->RegisterPrimitives(GetPrimitives(CompositeMeshActors));
 		Subsystem->SetCompositeTextureOverride(bOverrideCompositeTexture ? CompositeTexture : nullptr);
diff --git a/Engine/Plugins/Experimental/Compositing/HoldoutComposite/Source/HoldoutComposite/Private/HoldoutCompositeSceneViewExtension.cpp b/Engine/Plugins/Experimental/Compositing/HoldoutComposite/Source/HoldoutComposite/Private/HoldoutCompositeSceneViewExtension.cpp
--- a/Engine/Plugins/Experimental/Compositing/HoldoutComposite/Source/HoldoutComposite/Private/HoldoutCompositeSceneViewExtension.cpp
+++ b/Engine/Plugins/Experimental/Compositing/HoldoutComposite/Source/HoldoutComposite/Private/HoldoutCompositeSceneViewExtension.cpp
@@ -32,7 +32,7 @@ DECLARE_GPU_STAT_NAMED(FHoldoutCompositeFinal, TEXT("HoldoutComposite.Final"));
 
 namespace HoldoutComposite
 {
-	enum class ESceneColorSourceEncoding : uint32
+	enum class ESourceEncoding : uint32
 	{
 		Linear = 0,
 		Gamma = 1,
@@ -231,7 +231,7 @@ public:
 
 	virtual void OnPostRender(FRDGBuilder& GraphBuilder) override
 	{
-		const FRDGTextureDesc TextureDesc = FRDGTextureDesc::Create2D(RenderTargetTexture->Desc.Extent, PF_FloatRGBA, FClearValueBinding::Black, TexCreate_UAV | TexCreate_ShaderResource);
+		const FRDGTextureDesc TextureDesc = FRDGTextureDesc::Create2D(RenderTargetTexture->Desc.Extent, PF_FloatRGBA, FClearValueBinding::Black, TexCreate_UAV | TexCreate_ShaderResource | TexCreate_RenderTargetable);
 		FRDGTextureRef Output = GraphBuilder.CreateTexture(TextureDesc, TEXT("HoldoutCompositeProcessedTexture"));
 
 		HoldoutComposite::AddDilatePass(GraphBuilder, RenderTargetTexture, Output, ViewFeatureLevel);
@@ -249,17 +249,30 @@ private:
 FHoldoutCompositeSceneViewExtension::FHoldoutCompositeSceneViewExtension(const FAutoRegister& AutoReg, UWorld* InWorld)
 	: FWorldSceneViewExtension(AutoReg, InWorld)
 {
-	const UHoldoutCompositePluginSettings* PluginSettings = GetDefault<UHoldoutCompositePluginSettings>();
+	UHoldoutCompositePluginSettings* PluginSettings = GetMutableDefault<UHoldoutCompositePluginSettings>();
 	check(PluginSettings);
 
+#if WITH_EDITOR
+	DefautSettingsChangeHandle = PluginSettings->OnDefaultCompositeSettingsChanged.AddLambda([this]
+		{
+			OverrideCompositeSettings_GameThread(GetDefault<UHoldoutCompositePluginSettings>()->DefaultCompositeSettings, MIN_int32);
+		});
+#endif
+
 	CompositeSettings = CompositeSettings_RenderThread = PluginSettings->DefaultCompositeSettings;
 }
 
 FHoldoutCompositeSceneViewExtension::~FHoldoutCompositeSceneViewExtension()
 {
+#if WITH_EDITOR
+	if (UHoldoutCompositePluginSettings* PluginSettings = GetMutableDefault<UHoldoutCompositePluginSettings>())
+	{
+		PluginSettings->OnDefaultCompositeSettingsChanged.Remove(DefautSettingsChangeHandle);
+	}
+#endif
 }
 
-void FHoldoutCompositeSceneViewExtension::RegisterPrimitives_GameThread(const TArray<UPrimitiveComponent*>& InPrimitiveComponents, bool bInHoldoutState)
+void FHoldoutCompositeSceneViewExtension::RegisterPrimitives_GameThread(const TArray<UPrimitiveComponent*>& InPrimitiveComponents)
 {
 	check(IsInGameThread());
 
@@ -276,11 +289,11 @@ void FHoldoutCompositeSceneViewExtension::RegisterPrimitives_GameThread(const TA
 		}
 
 		// The SetHoldout() function makes changes only if the HoldoutState value differs from the current value.
-		InPrimitiveComponent->SetHoldout(bInHoldoutState);
+		InPrimitiveComponent->SetHoldout(true);
 	}
 }
 
-void FHoldoutCompositeSceneViewExtension::UnregisterPrimitives_GameThread(const TArray<UPrimitiveComponent*>& InPrimitiveComponents, bool bInHoldoutState)
+void FHoldoutCompositeSceneViewExtension::UnregisterPrimitives_GameThread(const TArray<UPrimitiveComponent*>& InPrimitiveComponents)
 {
 	check(IsInGameThread());
 
@@ -297,15 +310,18 @@ void FHoldoutCompositeSceneViewExtension::UnregisterPrimitives_GameThread(const
 		}
 
 		// The SetHoldout() function makes changes only if the HoldoutState value differs from the current value.
-		InPrimitiveComponent->SetHoldout(bInHoldoutState);
+		InPrimitiveComponent->SetHoldout(false);
 	}
 }
 
-void FHoldoutCompositeSceneViewExtension::OverrideCompositeSettings_GameThread(const FHoldoutCompositeSettings& InSettings)
+void FHoldoutCompositeSceneViewExtension::OverrideCompositeSettings_GameThread(const FHoldoutCompositeSettings& InSettings, int32 InPriority)
 {
 	check(IsInGameThread());
 
-	UpdateSettings_GameThread(InSettings);
+	if (InPriority >= CurrentSettingsPriority)
+	{
+		UpdateSettings_GameThread(InSettings, InPriority);
+	}
 }
 
 void FHoldoutCompositeSceneViewExtension::ResetCompositeSettings_GameThread()
@@ -315,7 +331,7 @@ void FHoldoutCompositeSceneViewExtension::ResetCompositeSettings_GameThread()
 	const UHoldoutCompositePluginSettings* PluginSettings = GetDefault<UHoldoutCompositePluginSettings>();
 	check(PluginSettings);
 
-	UpdateSettings_GameThread(PluginSettings->DefaultCompositeSettings);
+	UpdateSettings_GameThread(PluginSettings->DefaultCompositeSettings, MIN_int32);
 }
 
 void FHoldoutCompositeSceneViewExtension::SetCompositePostProcessMaterials_GameThread(const TArray<UMaterialInterface*>& InMaterials)
@@ -434,7 +450,7 @@ void FHoldoutCompositeSceneViewExtension::SetupView(FSceneViewFamily& InViewFami
 	PassInput.bOverridesPostVolumeUserFlags = GetSettings().bOverridesViewUserFlags;
 	PassInput.PostVolumeUserFlags = GetSettings().CompositeViewUserFlags;
 	// Note: Incoming view location is invalid for scene captures
-	PassInput.ViewLocation = InView.bIsSceneCapture ? InView.ViewMatrices.GetViewOrigin() : InView.ViewLocation;
+	PassInput.ViewLocation = InView.ViewMatrices.GetViewOrigin();
 	PassInput.ViewRotationMatrix = InView.ViewMatrices.GetViewMatrix().RemoveTranslation();
 	PassInput.ViewRotationMatrix.RemoveScaling();
 
@@ -472,8 +488,9 @@ void FHoldoutCompositeSceneViewExtension::SubscribeToPostProcessingPass(EPostPro
 	}
 }
 
-void FHoldoutCompositeSceneViewExtension::UpdateSettings_GameThread(const FHoldoutCompositeSettings& InSettings)
+void FHoldoutCompositeSceneViewExtension::UpdateSettings_GameThread(const FHoldoutCompositeSettings& InSettings, int32 InPriority)
 {
+	CurrentSettingsPriority = InPriority;
 	CompositeSettings = InSettings;
 	
 	ENQUEUE_RENDER_COMMAND(CopyHoldoutCompositeSettings)(
@@ -632,13 +649,13 @@ FScreenPassTexture FHoldoutCompositeSceneViewExtension::PostProcessCompositePass
 	FUint32Vector Encodings = FUint32Vector::ZeroValue; // ESceneColorSourceEncoding::Linear
 	if ((Family->EngineShowFlags.Tonemapper == 0) || (Family->EngineShowFlags.PostProcessing == 0))
 	{
-		Encodings.X = static_cast<uint32>(ESceneColorSourceEncoding::Gamma);
-		Encodings.Y = static_cast<uint32>(ESceneColorSourceEncoding::Gamma);
+		Encodings.X = static_cast<uint32>(ESourceEncoding::Gamma);
+		Encodings.Y = static_cast<uint32>(ESourceEncoding::Gamma);
 	}
 	else if (Family->SceneCaptureSource == SCS_FinalColorLDR)
 	{
-		Encodings.X = static_cast<uint32>(ESceneColorSourceEncoding::sRGB);
-		Encodings.Y = static_cast<uint32>(ESceneColorSourceEncoding::sRGB);
+		Encodings.X = static_cast<uint32>(ESourceEncoding::sRGB);
+		Encodings.Y = static_cast<uint32>(ESourceEncoding::sRGB);
 	}
 
 	// Encodings.Z controls the encoding of the composite texture
@@ -646,12 +663,12 @@ FScreenPassTexture FHoldoutCompositeSceneViewExtension::PostProcessCompositePass
 	if (!GetSettings().bLinearComposite)
 	{
 		// The composite texture is in the sRGB color space and must be encoded in the linear color space
-		Encodings.Z = static_cast<uint32>(ESceneColorSourceEncoding::sRGB);
+		Encodings.Z = static_cast<uint32>(ESourceEncoding::sRGB);
 	}
 	else if(GetSettings().bApplyFXAA)
 	{
 		// FXAA received gamma-encoded colors, so we must linearize before compositing.
-		Encodings.Z = static_cast<uint32>(ESceneColorSourceEncoding::Gamma);
+		Encodings.Z = static_cast<uint32>(ESourceEncoding::Gamma);
 	}
 
 
diff --git a/Engine/Plugins/Experimental/Compositing/HoldoutComposite/Source/HoldoutComposite/Private/HoldoutCompositeSceneViewExtension.h b/Engine/Plugins/Experimental/Compositing/HoldoutComposite/Source/HoldoutComposite/Private/HoldoutCompositeSceneViewExtension.h
--- a/Engine/Plugins/Experimental/Compositing/HoldoutComposite/Source/HoldoutComposite/Private/HoldoutCompositeSceneViewExtension.h
+++ b/Engine/Plugins/Experimental/Compositing/HoldoutComposite/Source/HoldoutComposite/Private/HoldoutCompositeSceneViewExtension.h
@@ -24,13 +24,13 @@ public:
 	~FHoldoutCompositeSceneViewExtension();
 
 	/* Register primitives for compositing. */
-	void RegisterPrimitives_GameThread(const TArray<UPrimitiveComponent*>& InPrimitiveComponents, bool bInHoldoutState=true);
+	void RegisterPrimitives_GameThread(const TArray<UPrimitiveComponent*>& InPrimitiveComponents);
 
 	/* Unregister primitives for compositing. */
-	void UnregisterPrimitives_GameThread(const TArray<UPrimitiveComponent*>& InPrimitiveComponents, bool bInHoldoutState=false);
+	void UnregisterPrimitives_GameThread(const TArray<UPrimitiveComponent*>& InPrimitiveComponents);
 
 	/* Override the composite settings. */
-	void OverrideCompositeSettings_GameThread(const FHoldoutCompositeSettings& InSettings);
+	void OverrideCompositeSettings_GameThread(const FHoldoutCompositeSettings& InSettings, int32 InPriority);
 
 	/* Reset composite settings back to plugin settings default. */
 	void ResetCompositeSettings_GameThread();
@@ -66,7 +66,7 @@ protected:
 	//~ End ISceneViewExtension Interface
 
 	/* Update composite settings. */
-	void UpdateSettings_GameThread(const FHoldoutCompositeSettings& InSettings);
+	void UpdateSettings_GameThread(const FHoldoutCompositeSettings& InSettings, int32 InPriority);
 
 	/* Thread-safe getter for composite settings. */
 	const FHoldoutCompositeSettings& GetSettings() const;
@@ -92,6 +92,9 @@ private:
 	FHoldoutCompositeSettings CompositeSettings;
 	FHoldoutCompositeSettings CompositeSettings_RenderThread;
 
+	/* Composite settings priority */
+	int32 CurrentSettingsPriority = MIN_int32;
+
 	/* Composite texture override use directly instead of the separate composited render. */
 	TWeakObjectPtr<UTexture> CompositeTextureOverride;
 
@@ -100,5 +103,10 @@ private:
 
 	/** Critical section for composite post-process materials. */
 	FCriticalSection MaterialsCS;
+
+#if WITH_EDITOR
+	/** Delegate handle for default plugin settings change. */
+	FDelegateHandle DefautSettingsChangeHandle;
+#endif
 };
 
diff --git a/Engine/Plugins/Experimental/Compositing/HoldoutComposite/Source/HoldoutComposite/Private/HoldoutCompositeSettings.cpp b/Engine/Plugins/Experimental/Compositing/HoldoutComposite/Source/HoldoutComposite/Private/HoldoutCompositeSettings.cpp
--- a/Engine/Plugins/Experimental/Compositing/HoldoutComposite/Source/HoldoutComposite/Private/HoldoutCompositeSettings.cpp
+++ b/Engine/Plugins/Experimental/Compositing/HoldoutComposite/Source/HoldoutComposite/Private/HoldoutCompositeSettings.cpp
@@ -1,6 +1,7 @@
 // Copyright Epic Games, Inc. All Rights Reserved.
 
 #include "HoldoutCompositeSettings.h"
+#include "HoldoutCompositeSubsystem.h"
 
 #include "Components/PrimitiveComponent.h"
 
@@ -21,7 +22,15 @@ FHoldoutCompositeSettings::FHoldoutCompositeSettings()
 { }
 
 UHoldoutCompositePluginSettings::UHoldoutCompositePluginSettings()
-	: DisabledPrimitiveClasses()
+	: DisabledPrimitiveClasses(
+		{
+			FSoftClassPath("/Script/Engine.CapsuleComponent"),
+			FSoftClassPath("/Script/Engine.BillboardComponent"),
+			FSoftClassPath("/Script/Engine.ArrowComponent"),
+			FSoftClassPath("/Script/Engine.DrawFrustumComponent"),
+			FSoftClassPath("/Script/Engine.LineBatchComponent")
+		}
+	)
 	, DefaultCompositeSettings()
 { }
 
@@ -40,6 +49,15 @@ FName UHoldoutCompositePluginSettings::GetSectionName() const
 {
 	return TEXT("Holdout Composite");
 }
+void UHoldoutCompositePluginSettings::PostEditChangeProperty(FPropertyChangedEvent& PropertyChangedEvent)
+{
+	Super::PostEditChangeProperty(PropertyChangedEvent);
+
+	if (PropertyChangedEvent.GetMemberPropertyName() == GET_MEMBER_NAME_CHECKED(ThisClass, DefaultCompositeSettings))
+	{
+		OnDefaultCompositeSettingsChanged.Broadcast();
+	}
+}
 #endif
 
 bool UHoldoutCompositePluginSettings::IsAllowedPrimitiveClass(const UPrimitiveComponent* InPrimitiveComponent) const
diff --git a/Engine/Plugins/Experimental/Compositing/HoldoutComposite/Source/HoldoutComposite/Private/HoldoutCompositeSubsystem.cpp b/Engine/Plugins/Experimental/Compositing/HoldoutComposite/Source/HoldoutComposite/Private/HoldoutCompositeSubsystem.cpp
--- a/Engine/Plugins/Experimental/Compositing/HoldoutComposite/Source/HoldoutComposite/Private/HoldoutCompositeSubsystem.cpp
+++ b/Engine/Plugins/Experimental/Compositing/HoldoutComposite/Source/HoldoutComposite/Private/HoldoutCompositeSubsystem.cpp
@@ -64,35 +64,35 @@ void UHoldoutCompositeSubsystem::Deinitialize()
 	Super::Deinitialize();
 }
 
-void UHoldoutCompositeSubsystem::RegisterPrimitive(UPrimitiveComponent* InPrimitiveComponent, bool bInHoldoutState)
+void UHoldoutCompositeSubsystem::RegisterPrimitive(UPrimitiveComponent* InPrimitiveComponent)
 {
 	TArray<UPrimitiveComponent*> PrimitiveComponents = { InPrimitiveComponent };
 
-	RegisterPrimitives(PrimitiveComponents, bInHoldoutState);
+	RegisterPrimitives(PrimitiveComponents);
 }
 
-void UHoldoutCompositeSubsystem::RegisterPrimitives(const TArray<UPrimitiveComponent*>& InPrimitiveComponents, bool bInHoldoutState)
+void UHoldoutCompositeSubsystem::RegisterPrimitives(const TArray<UPrimitiveComponent*>& InPrimitiveComponents)
 {
 	const bool bValidSettings = ValidateProjectSettings();
 
 	if (bValidSettings && HoldoutCompositeViewExtension.IsValid())
 	{
-		HoldoutCompositeViewExtension->RegisterPrimitives_GameThread(InPrimitiveComponents, bInHoldoutState);
+		HoldoutCompositeViewExtension->RegisterPrimitives_GameThread(InPrimitiveComponents);
 	}
 }
 
-void UHoldoutCompositeSubsystem::UnregisterPrimitive(UPrimitiveComponent* InPrimitiveComponent, bool bInHoldoutState)
+void UHoldoutCompositeSubsystem::UnregisterPrimitive(UPrimitiveComponent* InPrimitiveComponent)
 {
 	TArray<UPrimitiveComponent*> PrimitiveComponents = { InPrimitiveComponent };
 
-	UnregisterPrimitives(PrimitiveComponents, bInHoldoutState);
+	UnregisterPrimitives(PrimitiveComponents);
 }
 
-void UHoldoutCompositeSubsystem::UnregisterPrimitives(const TArray<UPrimitiveComponent*>& InPrimitiveComponents, bool bInHoldoutState)
+void UHoldoutCompositeSubsystem::UnregisterPrimitives(const TArray<UPrimitiveComponent*>& InPrimitiveComponents)
 {
 	if (HoldoutCompositeViewExtension.IsValid())
 	{
-		HoldoutCompositeViewExtension->UnregisterPrimitives_GameThread(InPrimitiveComponents, bInHoldoutState);
+		HoldoutCompositeViewExtension->UnregisterPrimitives_GameThread(InPrimitiveComponents);
 	}
 }
 
@@ -112,11 +112,11 @@ void UHoldoutCompositeSubsystem::ResetCompositePostProcessMaterials()
 	}
 }
 
-void UHoldoutCompositeSubsystem::OverrideCompositeSettings(const FHoldoutCompositeSettings& InSettings)
+void UHoldoutCompositeSubsystem::OverrideCompositeSettings(const FHoldoutCompositeSettings& InSettings, int32 InPriority)
 {
 	if (HoldoutCompositeViewExtension.IsValid())
 	{
-		HoldoutCompositeViewExtension->OverrideCompositeSettings_GameThread(InSettings);
+		HoldoutCompositeViewExtension->OverrideCompositeSettings_GameThread(InSettings, InPriority);
 	}
 }
 
diff --git a/Engine/Plugins/Experimental/Compositing/HoldoutComposite/Source/HoldoutComposite/Public/HoldoutCompositeSettings.h b/Engine/Plugins/Experimental/Compositing/HoldoutComposite/Source/HoldoutComposite/Public/HoldoutCompositeSettings.h
--- a/Engine/Plugins/Experimental/Compositing/HoldoutComposite/Source/HoldoutComposite/Public/HoldoutCompositeSettings.h
+++ b/Engine/Plugins/Experimental/Compositing/HoldoutComposite/Source/HoldoutComposite/Public/HoldoutCompositeSettings.h
@@ -9,6 +9,10 @@
 
 #include "HoldoutCompositeSettings.generated.h"
 
+#if WITH_EDITOR
+DECLARE_MULTICAST_DELEGATE(FOnDefaultCompositeSettingsChanged);
+#endif
+
 class UPrimitiveComponent;
 
 /* Post-processing location of the main holdout composite pass. */
@@ -67,7 +71,7 @@ struct HOLDOUTCOMPOSITE_API FHoldoutCompositeSettings
 	int32 CompositeViewUserFlags;
 
 	/** If enabled, composite renders in linear colorspace. */
-	UPROPERTY(config, EditAnywhere, Category = HoldoutComposite, AdvancedDisplay)
+	UPROPERTY(config)
 	bool bLinearComposite;
 };
 
@@ -87,6 +91,7 @@ public:
 #if WITH_EDITOR
 	virtual FText GetSectionText() const override;
 	virtual FName GetSectionName() const override;
+	virtual void PostEditChangeProperty(FPropertyChangedEvent& PropertyChangedEvent) override;
 #endif
 	//~ End UDeveloperSettings interface
 
@@ -102,6 +107,11 @@ public:
 	/** Default holdout composite settings.*/
 	UPROPERTY(config, EditAnywhere, Category = General, meta = (ShowOnlyInnerProperties))
 	FHoldoutCompositeSettings DefaultCompositeSettings;
+
+#if WITH_EDITOR
+	/** Delegate event that triggers when the default settings are changed.*/
+	FOnDefaultCompositeSettingsChanged OnDefaultCompositeSettingsChanged;
+#endif
 };
 
 
diff --git a/Engine/Plugins/Experimental/Compositing/HoldoutComposite/Source/HoldoutComposite/Public/HoldoutCompositeSubsystem.h b/Engine/Plugins/Experimental/Compositing/HoldoutComposite/Source/HoldoutComposite/Public/HoldoutCompositeSubsystem.h
--- a/Engine/Plugins/Experimental/Compositing/HoldoutComposite/Source/HoldoutComposite/Public/HoldoutCompositeSubsystem.h
+++ b/Engine/Plugins/Experimental/Compositing/HoldoutComposite/Source/HoldoutComposite/Public/HoldoutCompositeSubsystem.h
@@ -31,17 +31,17 @@ public:
 
 	/* Register a single primitive for compositing. */
 	UFUNCTION(BlueprintCallable, Category = "Holdout Composite")
-	virtual void RegisterPrimitive(UPrimitiveComponent* InPrimitiveComponent, bool bInHoldoutState=true);
+	virtual void RegisterPrimitive(UPrimitiveComponent* InPrimitiveComponent);
 
 	/* Register multiple primitives for compositing. */
-	virtual void RegisterPrimitives(const TArray<UPrimitiveComponent*>& InPrimitiveComponents, bool bInHoldoutState=true);
+	virtual void RegisterPrimitives(const TArray<UPrimitiveComponent*>& InPrimitiveComponents);
 
 	/* Unregister a single primitive from compositing. */
 	UFUNCTION(BlueprintCallable, Category = "Holdout Composite")
-	virtual void UnregisterPrimitive(UPrimitiveComponent* InPrimitiveComponent, bool bInHoldoutState=false);
+	virtual void UnregisterPrimitive(UPrimitiveComponent* InPrimitiveComponent);
 
 	/* Unregister multiple primitives from compositing. */
-	virtual void UnregisterPrimitives(const TArray<UPrimitiveComponent*>& InPrimitiveComponents, bool bInHoldoutState=false);
+	virtual void UnregisterPrimitives(const TArray<UPrimitiveComponent*>& InPrimitiveComponents);
 
 	/* Set additional post-process materials to be applied on the separate composited render. */
 	UFUNCTION(BlueprintCallable, Category = "Holdout Composite")
@@ -53,7 +53,7 @@ public:
 
 	/* Override project default composite settings for the active holdout composite passes. */
 	UFUNCTION(BlueprintCallable, Category = "Holdout Composite")
-	virtual void OverrideCompositeSettings(const FHoldoutCompositeSettings& InSettings);
+	virtual void OverrideCompositeSettings(const FHoldoutCompositeSettings& InSettings, int32 InPriority);
 
 	/* Revert composite settings back to project defaults. */
 	UFUNCTION(BlueprintCallable, Category = "Holdout Composite")
