HoldoutComposite: Fix SVE incorrectly running inside scene captures that had the primitive alpha holdout showflag enabled.

diff --git a/Engine/Plugins/Experimental/Compositing/HoldoutComposite/Source/HoldoutComposite/Private/HoldoutCompositeSceneViewExtension.cpp b/Engine/Plugins/Experimental/Compositing/HoldoutComposite/Source/HoldoutComposite/Private/HoldoutCompositeSceneViewExtension.cpp
--- a/Engine/Plugins/Experimental/Compositing/HoldoutComposite/Source/HoldoutComposite/Private/HoldoutCompositeSceneViewExtension.cpp
+++ b/Engine/Plugins/Experimental/Compositing/HoldoutComposite/Source/HoldoutComposite/Private/HoldoutCompositeSceneViewExtension.cpp
@@ -386,6 +386,12 @@ void FHoldoutCompositeSceneViewExtension::SetupView(FSceneViewFamily& InViewFami
 		return;
 	}
 
+	// Disable custom render pass for view families without primitive alpha holdout support.
+	if (InViewFamily.EngineShowFlags.AllowPrimitiveAlphaHoldout == 0)
+	{
+		return;
+	}
+
 	const TWeakObjectPtr<UWorld> WorldPtr = GetWorld();
 	check(WorldPtr.IsValid());
 
@@ -468,7 +474,10 @@ void FHoldoutCompositeSceneViewExtension::SubscribeToPostProcessingPass(EPostPro
 {
 	constexpr bool bEvenIfPendingKill = false;
 	constexpr bool bThreadsafeTest = true;
-	if (!CustomRenderTargetPerView_RenderThread.Contains(InView.GetViewKey()) && !CompositeTextureOverride.IsValid(bEvenIfPendingKill, bThreadsafeTest))
+	bool bIsActive = static_cast<bool>(InView.Family->EngineShowFlags.AllowPrimitiveAlphaHoldout) &&
+		(CustomRenderTargetPerView_RenderThread.Contains(InView.GetViewKey()) || CompositeTextureOverride.IsValid(bEvenIfPendingKill, bThreadsafeTest));
+
+	if (!bIsActive)
 	{
 		// Early-out to avoid needless work in the post processing callback(s).
 		return;
