HoldoutComposite: Re-register primitives on tick by default for both the actor and the component classes.

* This fixes Text3D use cases, where the primitives are dynamically regenerated.
* We precache the list of disallowed classes to avoid load object call cost on tick.

diff --git a/Engine/Plugins/Experimental/Compositing/HoldoutComposite/Source/HoldoutComposite/Private/HoldoutComposite.cpp b/Engine/Plugins/Experimental/Compositing/HoldoutComposite/Source/HoldoutComposite/Private/HoldoutComposite.cpp
--- a/Engine/Plugins/Experimental/Compositing/HoldoutComposite/Source/HoldoutComposite/Private/HoldoutComposite.cpp
+++ b/Engine/Plugins/Experimental/Compositing/HoldoutComposite/Source/HoldoutComposite/Private/HoldoutComposite.cpp
@@ -45,6 +45,24 @@ namespace UE
 	}
 }
 
+AHoldoutComposite::AHoldoutComposite(const FObjectInitializer& ObjectInitializer)
+	: Super(ObjectInitializer)
+{
+	if (UE::HoldoutComposite::Private::GRegisterPrimitivesOnTick)
+	{
+		PrimaryActorTick.bCanEverTick = true;
+		PrimaryActorTick.bTickEvenWhenPaused = true;
+		PrimaryActorTick.TickGroup = TG_PostUpdateWork;
+
+		PrimaryActorTick.bStartWithTickEnabled = true;
+		bAllowTickBeforeBeginPlay = true;
+	}
+
+	SetActorEnableCollision(false);
+}
+
+AHoldoutComposite::~AHoldoutComposite() = default;
+
 void AHoldoutComposite::PostActorCreated()
 {
 	Super::PostActorCreated();
@@ -83,6 +101,14 @@ void AHoldoutComposite::PostUnregisterAllComponents()
 	UnregisterCompositeImpl();
 }
 
+void AHoldoutComposite::Tick(float DeltaSeconds)
+{
+	Super::Tick(DeltaSeconds);
+
+	constexpr bool bIsTick = true;
+	RegisterCompositeImpl(bIsTick);
+}
+
 void AHoldoutComposite::BeginDestroy()
 {
 	Super::BeginDestroy();
@@ -112,6 +138,10 @@ void AHoldoutComposite::SetIsTemporarilyHiddenInEditor(bool bIsHidden)
 		RegisterCompositeImpl();
 	}
 }
+bool AHoldoutComposite::ShouldTickIfViewportsOnly() const
+{
+	return UE::HoldoutComposite::Private::GRegisterPrimitivesOnTick ? true : Super::ShouldTickIfViewportsOnly();
+}
 #endif
 
 bool AHoldoutComposite::IsEnabled() const
@@ -133,15 +163,27 @@ void AHoldoutComposite::SetEnabled(bool bInEnabled)
 	}
 }
 
-void AHoldoutComposite::RegisterCompositeImpl()
+void AHoldoutComposite::RegisterCompositeImpl(bool bIsTick)
 {
 	using namespace UE::HoldoutComposite::Private;
 
+	if (GRegisterPrimitivesOnTick != bIsTick)
+	{
+		return;
+	}
+
 	if (!bIsEnabled)
 	{
 		return;
 	}
 
+#if WITH_EDITOR
+	if (IsTemporarilyHiddenInEditor())
+	{
+		return;
+	}
+#endif
+
 	UHoldoutCompositeSubsystem* Subsystem = UWorld::GetSubsystem<UHoldoutCompositeSubsystem>(GetWorld());
 	if (IsValid(Subsystem))
 	{
diff --git a/Engine/Plugins/Experimental/Compositing/HoldoutComposite/Source/HoldoutComposite/Private/HoldoutCompositeComponent.cpp b/Engine/Plugins/Experimental/Compositing/HoldoutComposite/Source/HoldoutComposite/Private/HoldoutCompositeComponent.cpp
--- a/Engine/Plugins/Experimental/Compositing/HoldoutComposite/Source/HoldoutComposite/Private/HoldoutCompositeComponent.cpp
+++ b/Engine/Plugins/Experimental/Compositing/HoldoutComposite/Source/HoldoutComposite/Private/HoldoutCompositeComponent.cpp
@@ -63,6 +63,15 @@ namespace UE
 UHoldoutCompositeComponent::UHoldoutCompositeComponent(const FObjectInitializer& ObjectInitializer)
 	: Super(ObjectInitializer)
 {
+	if (UE::HoldoutComposite::Private::GRegisterPrimitivesOnTick)
+	{
+		PrimaryComponentTick.TickGroup = TG_PostUpdateWork;
+		PrimaryComponentTick.bStartWithTickEnabled = true;
+		PrimaryComponentTick.bCanEverTick = true;
+		PrimaryComponentTick.bTickEvenWhenPaused = true;
+		bTickInEditor = true;
+		bAutoActivate = true;
+	}
 }
 
 void UHoldoutCompositeComponent::OnRegister()
@@ -79,6 +88,14 @@ void UHoldoutCompositeComponent::OnUnregister()
 	Super::OnUnregister();
 }
 
+void UHoldoutCompositeComponent::TickComponent(float DeltaTime, ELevelTick TickType, FActorComponentTickFunction* ThisTickFunction)
+{
+	Super::TickComponent(DeltaTime, TickType, ThisTickFunction);
+
+	constexpr bool bIsTick = true;
+	RegisterCompositeImpl(bIsTick);
+}
+
 void UHoldoutCompositeComponent::DetachFromComponent(const FDetachmentTransformRules& DetachmentRules)
 {
 	// Note: We also unregister here while the attached parent pointer is still valid.
@@ -136,15 +153,22 @@ void UHoldoutCompositeComponent::SetEnabled(bool bInIsEnabled)
 	}
 }
 
-void UHoldoutCompositeComponent::RegisterCompositeImpl()
+void UHoldoutCompositeComponent::RegisterCompositeImpl(bool bIsTick)
 {
+	using namespace UE::HoldoutComposite::Private;
+
+	if (GRegisterPrimitivesOnTick != bIsTick)
+	{
+		return;
+	}
+
 	if (!bIsEnabled)
 	{
 		return;
 	}
 
 	UHoldoutCompositeSubsystem* Subsystem = UWorld::GetSubsystem<UHoldoutCompositeSubsystem>(GetWorld());
-	TArray<UPrimitiveComponent*> ParentPrimitives = UE::HoldoutComposite::Private::FindPrimitiveComponents(GetAttachParent());
+	TArray<UPrimitiveComponent*> ParentPrimitives = FindPrimitiveComponents(GetAttachParent());
 
 	if (IsValid(Subsystem) && !ParentPrimitives.IsEmpty())
 	{
@@ -154,8 +178,10 @@ void UHoldoutCompositeComponent::RegisterCompositeImpl()
 
 void UHoldoutCompositeComponent::UnregisterCompositeImpl()
 {
+	using namespace UE::HoldoutComposite::Private;
+
 	UHoldoutCompositeSubsystem* Subsystem = UWorld::GetSubsystem<UHoldoutCompositeSubsystem>(GetWorld());
-	TArray<UPrimitiveComponent*> ParentPrimitives = UE::HoldoutComposite::Private::FindPrimitiveComponents(GetAttachParent());
+	TArray<UPrimitiveComponent*> ParentPrimitives = FindPrimitiveComponents(GetAttachParent());
 
 	if (IsValid(Subsystem) && !ParentPrimitives.IsEmpty())
 	{
diff --git a/Engine/Plugins/Experimental/Compositing/HoldoutComposite/Source/HoldoutComposite/Private/HoldoutCompositeSettings.cpp b/Engine/Plugins/Experimental/Compositing/HoldoutComposite/Source/HoldoutComposite/Private/HoldoutCompositeSettings.cpp
--- a/Engine/Plugins/Experimental/Compositing/HoldoutComposite/Source/HoldoutComposite/Private/HoldoutCompositeSettings.cpp
+++ b/Engine/Plugins/Experimental/Compositing/HoldoutComposite/Source/HoldoutComposite/Private/HoldoutCompositeSettings.cpp
@@ -7,6 +7,13 @@
 
 #include UE_INLINE_GENERATED_CPP_BY_NAME(HoldoutCompositeSettings)
 
+bool UE::HoldoutComposite::Private::GRegisterPrimitivesOnTick = true;
+static FAutoConsoleVariableRef CVarRefreshPrimitivesOnTick(
+	TEXT("HoldoutComposite.RefreshPrimitivesOnTick"),
+	UE::HoldoutComposite::Private::GRegisterPrimitivesOnTick,
+	TEXT("Refresh primitive registration on every tick. Enabled by default for use cases such as dynamic text 3D.\n")
+	TEXT("Read-only and to be set in a config file (requires restart)."),
+	ECVF_ReadOnly);
 
 FHoldoutCompositeSettings::FHoldoutCompositeSettings()
 	: bFollowSceneExposure(false)
@@ -57,15 +64,25 @@ void UHoldoutCompositePluginSettings::PostEditChangeProperty(FPropertyChangedEve
 	{
 		OnDefaultCompositeSettingsChanged.Broadcast();
 	}
+	else if (PropertyChangedEvent.GetMemberPropertyName() == GET_MEMBER_NAME_CHECKED(ThisClass, DisabledPrimitiveClasses))
+	{
+		CacheDisabledClasses();
+	}
 }
 #endif
 
+void UHoldoutCompositePluginSettings::PostInitProperties()
+{
+	Super::PostInitProperties();
+
+	CacheDisabledClasses();
+}
+
 bool UHoldoutCompositePluginSettings::IsAllowedPrimitiveClass(const UPrimitiveComponent* InPrimitiveComponent) const
 {
-	for (const FSoftClassPath& ObjectClass : DisabledPrimitiveClasses)
+	for (const UClass* ObjectClass : CachedDisabledPrimitiveClasses)
 	{
-		UClass* TransactionClass = ObjectClass.TryLoadClass<UObject>();
-		if (TransactionClass && InPrimitiveComponent->IsA(TransactionClass))
+		if (InPrimitiveComponent->IsA(ObjectClass))
 		{
 			return false;
 		}
@@ -73,3 +90,16 @@ bool UHoldoutCompositePluginSettings::IsAllowedPrimitiveClass(const UPrimitiveCo
 
 	return true;
 }
+
+void UHoldoutCompositePluginSettings::CacheDisabledClasses()
+{
+	CachedDisabledPrimitiveClasses.Reset(0);
+
+	for (const FSoftClassPath& ObjectClass : DisabledPrimitiveClasses)
+	{
+		if (const UClass* DisabledClass = ObjectClass.TryLoadClass<UObject>())
+		{
+			CachedDisabledPrimitiveClasses.Add(DisabledClass);
+		}
+	}
+}
diff --git a/Engine/Plugins/Experimental/Compositing/HoldoutComposite/Source/HoldoutComposite/Private/HoldoutCompositeSubsystem.cpp b/Engine/Plugins/Experimental/Compositing/HoldoutComposite/Source/HoldoutComposite/Private/HoldoutCompositeSubsystem.cpp
--- a/Engine/Plugins/Experimental/Compositing/HoldoutComposite/Source/HoldoutComposite/Private/HoldoutCompositeSubsystem.cpp
+++ b/Engine/Plugins/Experimental/Compositing/HoldoutComposite/Source/HoldoutComposite/Private/HoldoutCompositeSubsystem.cpp
@@ -143,7 +143,7 @@ bool UHoldoutCompositeSubsystem::EnableHoldoutSettingsNotification() const
 
 bool UHoldoutCompositeSubsystem::IsProjectSettingsValid()
 {
-	URendererSettings* RendererSettings = GetMutableDefault<URendererSettings>();
+	const URendererSettings* RendererSettings = GetDefault<URendererSettings>();
 	check(RendererSettings);
 
 	const bool bValidSettings = 
diff --git a/Engine/Plugins/Experimental/Compositing/HoldoutComposite/Source/HoldoutComposite/Public/HoldoutComposite.h b/Engine/Plugins/Experimental/Compositing/HoldoutComposite/Source/HoldoutComposite/Public/HoldoutComposite.h
--- a/Engine/Plugins/Experimental/Compositing/HoldoutComposite/Source/HoldoutComposite/Public/HoldoutComposite.h
+++ b/Engine/Plugins/Experimental/Compositing/HoldoutComposite/Source/HoldoutComposite/Public/HoldoutComposite.h
@@ -19,10 +19,14 @@ class HOLDOUTCOMPOSITE_API AHoldoutComposite : public AActor
 
 public:
 
+	AHoldoutComposite(const FObjectInitializer& ObjectInitializer);
+	~AHoldoutComposite();
+
 	//~ Begin AActor Interface
 	virtual void PostActorCreated();
 	virtual void PostRegisterAllComponents() override;
 	virtual void PostUnregisterAllComponents() override;
+	virtual void Tick(float DeltaSeconds) override;
 	virtual void BeginDestroy() override;
 	//~ End AActor Interface
 
@@ -30,6 +34,7 @@ public:
 	virtual void PreEditChange(FProperty* PropertyThatWillChange) override;
 	virtual void PostEditChangeProperty(FPropertyChangedEvent& PropertyChangedEvent) override;
 	virtual void SetIsTemporarilyHiddenInEditor(bool bIsHidden) override;
+	virtual bool ShouldTickIfViewportsOnly() const override;
 #endif // WITH_EDITOR
 
 	/* Get the enabled state of the component. */
@@ -75,7 +80,7 @@ public:
 private:
 
 	/* Private implementation of the register method. */
-	void RegisterCompositeImpl();
+	void RegisterCompositeImpl(bool bIsTick = false);
 
 	/* Private implementation of the unregister method. */
 	void UnregisterCompositeImpl();
diff --git a/Engine/Plugins/Experimental/Compositing/HoldoutComposite/Source/HoldoutComposite/Public/HoldoutCompositeComponent.h b/Engine/Plugins/Experimental/Compositing/HoldoutComposite/Source/HoldoutComposite/Public/HoldoutCompositeComponent.h
--- a/Engine/Plugins/Experimental/Compositing/HoldoutComposite/Source/HoldoutComposite/Public/HoldoutCompositeComponent.h
+++ b/Engine/Plugins/Experimental/Compositing/HoldoutComposite/Source/HoldoutComposite/Public/HoldoutCompositeComponent.h
@@ -7,20 +7,21 @@
 
 #include "HoldoutCompositeComponent.generated.h"
 
-UCLASS(ClassGroup = Rendering, HideCategories=(Activation, Transform, Lighting, Rendering, Tags, Cooking, Physics, LOD, AssetUserData, Navigation), editinlinenew, meta = (BlueprintSpawnableComponent), MinimalAPI)
-class UHoldoutCompositeComponent : public USceneComponent
+UCLASS(ClassGroup = Rendering, HideCategories=(Activation, Transform, Lighting, Rendering, Tags, Cooking, Physics, LOD, AssetUserData, Navigation), editinlinenew, meta = (BlueprintSpawnableComponent))
+class HOLDOUTCOMPOSITE_API UHoldoutCompositeComponent : public USceneComponent
 {
 	GENERATED_UCLASS_BODY()
 
 public:
 	//~ Begin UActorComponent interface
-	HOLDOUTCOMPOSITE_API virtual void OnRegister() override;
-	HOLDOUTCOMPOSITE_API virtual void OnUnregister() override;
+	virtual void OnRegister() override;
+	virtual void OnUnregister() override;
+	virtual void TickComponent(float DeltaTime, enum ELevelTick TickType, FActorComponentTickFunction* ThisTickFunction) override;
 	//~ End UActorComponent interface
 	
 	//~ Begin USceneComponent Interface.
-	HOLDOUTCOMPOSITE_API virtual void OnAttachmentChanged() override;
-	HOLDOUTCOMPOSITE_API virtual void DetachFromComponent(const FDetachmentTransformRules& DetachmentRules) override;
+	virtual void OnAttachmentChanged() override;
+	virtual void DetachFromComponent(const FDetachmentTransformRules& DetachmentRules) override;
 	//~ End USceneComponent Interface
 
 	/* Get the enabled state of the component. */
@@ -34,7 +35,7 @@ public:
 private:
 
 	/* Private implementation of the register method. */
-	void RegisterCompositeImpl();
+	void RegisterCompositeImpl(bool bIsTick = false);
 
 	/* Private implementation of the unregister method. */
 	void UnregisterCompositeImpl();
diff --git a/Engine/Plugins/Experimental/Compositing/HoldoutComposite/Source/HoldoutComposite/Public/HoldoutCompositeSettings.h b/Engine/Plugins/Experimental/Compositing/HoldoutComposite/Source/HoldoutComposite/Public/HoldoutCompositeSettings.h
--- a/Engine/Plugins/Experimental/Compositing/HoldoutComposite/Source/HoldoutComposite/Public/HoldoutCompositeSettings.h
+++ b/Engine/Plugins/Experimental/Compositing/HoldoutComposite/Source/HoldoutComposite/Public/HoldoutCompositeSettings.h
@@ -15,6 +15,17 @@ DECLARE_MULTICAST_DELEGATE(FOnDefaultCompositeSettingsChanged);
 
 class UPrimitiveComponent;
 
+namespace UE
+{
+	namespace HoldoutComposite
+	{
+		namespace Private
+		{
+			extern bool GRegisterPrimitivesOnTick;
+		}
+	}
+}
+
 /* Post-processing location of the main holdout composite pass. */
 UENUM()
 enum class EHoldoutCompositePassLocation
@@ -95,6 +106,8 @@ public:
 #endif
 	//~ End UDeveloperSettings interface
 
+	virtual void PostInitProperties() override;
+
 public:
 
 	// Ensure that the primitive class is allowed by checking against plugin settings.
@@ -112,6 +125,13 @@ public:
 	/** Delegate event that triggers when the default settings are changed.*/
 	FOnDefaultCompositeSettingsChanged OnDefaultCompositeSettingsChanged;
 #endif
+private:
+
+	/** Cache disable class pointers to prevent load object calls during allowed class checks. */
+	void CacheDisabledClasses();
+	
+	/** Cached list of disallowed class types. */
+	TArray<const UClass*> CachedDisabledPrimitiveClasses;
 };
 
 
