HoldoutComposite: Adding actor built-in support for composite texture projection.

By default, the projection is done from the currently active camera, though users can specify a projection camera source. Once "Composite Texture Projection" is enabled, we convert the composite mesh materials into MIDs and continuously update four material parameters called "XPlane", "YPlane", "ZPlane" & "WPlane" representing the view-projection matrix. The composite texture is also automatically bound to "CompositeTexture" material parameters, if present.

* Added example/reference material & material function.
* Minor refactor of blueprint-exposed functions & tick method, and renaming.

#virtualized

diff --git a/Engine/Plugins/Experimental/Compositing/HoldoutComposite/Config/BaseHoldoutComposite.ini b/Engine/Plugins/Experimental/Compositing/HoldoutComposite/Config/BaseHoldoutComposite.ini
new file mode 100644
--- /dev/null
+++ b/Engine/Plugins/Experimental/Compositing/HoldoutComposite/Config/BaseHoldoutComposite.ini
@@ -0,0 +1,11 @@
+[CoreRedirects]
+
+;
+; 5.6 Redirects
+;
+
+; HoldoutComposite properties
++PropertyRedirects=(OldName="HoldoutComposite.bOverrideCompositeTexture",NewName="HoldoutComposite.bCompositeTextureOverrideRender")
++PropertyRedirects=(OldName="HoldoutComposite.bOverrideDefaultSettings",NewName="HoldoutComposite.bEnableOverrideSettings")
++PropertyRedirects=(OldName="HoldoutComposite.SettingsOverride",NewName="HoldoutComposite.OverrideSettings")
+
diff --git a/Engine/Plugins/Experimental/Compositing/HoldoutComposite/HoldoutComposite.uplugin b/Engine/Plugins/Experimental/Compositing/HoldoutComposite/HoldoutComposite.uplugin
--- a/Engine/Plugins/Experimental/Compositing/HoldoutComposite/HoldoutComposite.uplugin
+++ b/Engine/Plugins/Experimental/Compositing/HoldoutComposite/HoldoutComposite.uplugin
@@ -11,7 +11,7 @@
 	"MarketplaceURL": "",
 	"SupportURL": "",
 	"EnabledByDefault": false,
-	"CanContainContent": false,
+	"CanContainContent": true,
 	"IsBetaVersion": false,
 	"IsExperimentalVersion": true,
 	"Installed": false,
diff --git a/Engine/Plugins/Experimental/Compositing/HoldoutComposite/Source/HoldoutComposite/Private/HoldoutComposite.cpp b/Engine/Plugins/Experimental/Compositing/HoldoutComposite/Source/HoldoutComposite/Private/HoldoutComposite.cpp
--- a/Engine/Plugins/Experimental/Compositing/HoldoutComposite/Source/HoldoutComposite/Private/HoldoutComposite.cpp
+++ b/Engine/Plugins/Experimental/Compositing/HoldoutComposite/Source/HoldoutComposite/Private/HoldoutComposite.cpp
@@ -6,13 +6,26 @@
 #include "HoldoutCompositeSettings.h"
 
 #include "Components/PrimitiveComponent.h"
+#include "Camera/CameraActor.h"
+#include "Camera/CameraComponent.h"
 #include "Engine/Engine.h"
 #include "Engine/World.h"
 #include "Engine/Texture.h"
+#include "Kismet/GameplayStatics.h"
+#include "Materials/MaterialInstanceDynamic.h"
+#include "SceneView.h"
 #include "UObject/UObjectIterator.h"
 
+#if WITH_EDITOR
+#include "Editor.h"
+#include "LevelEditorViewport.h"
+#endif
+
 #include UE_INLINE_GENERATED_CPP_BY_NAME(HoldoutComposite)
 
+DECLARE_CYCLE_STAT(TEXT("HoldoutComp - Actor Tick"), STAT_HoldoutComposite_Actor_Tick, STATGROUP_HoldoutComposite);
+DECLARE_CYCLE_STAT(TEXT("HoldoutComp - Actor EndOfFrame Update"), STAT_HoldoutComposite_Actor_EndOfFrameUpdate, STATGROUP_HoldoutComposite);
+
 namespace UE
 {
 	namespace HoldoutComposite
@@ -89,24 +102,154 @@ void AHoldoutComposite::PostActorCreated()
 
 void AHoldoutComposite::PostRegisterAllComponents()
 {
+	using namespace UE::HoldoutComposite::Private;
+
 	Super::PostRegisterAllComponents();
 
-	RegisterCompositeImpl();
+	if (IsActive())
+	{
+		ApplyAllState();
+	}
+
+	/*
+	* NOTE: We must (unfortunately) rely on this after-tick event to update the view matrix projection,
+	* since viewport camera transform updates after ticking all world objects.
+	*
+	* The CompositePlane plugin updated this information on tick, but a lag becomes visible when
+	* piloting the source projection camera for example.
+	*/
+	OnWorldPreSendAllEndOfFrameUpdatesHandle = FWorldDelegates::OnWorldPreSendAllEndOfFrameUpdates.AddUObject(this, &AHoldoutComposite::EndOfFrameUpdate);
+
 }
 
 void AHoldoutComposite::PostUnregisterAllComponents()
 {
+	using namespace UE::HoldoutComposite::Private;
+
+	FWorldDelegates::OnWorldPreSendAllEndOfFrameUpdates.Remove(OnWorldPreSendAllEndOfFrameUpdatesHandle);
+
+	RemoveAllState();
+
 	Super::PostUnregisterAllComponents();
+}
+
+void AHoldoutComposite::Tick(float DeltaTime)
+{
+	Super::Tick(DeltaTime);
+
+	SCOPE_CYCLE_COUNTER(STAT_HoldoutComposite_Actor_Tick);
+
+	if (IsActive() && UE::HoldoutComposite::Private::GRegisterPrimitivesOnTick)
+	{
+		UHoldoutCompositeSubsystem* Subsystem = UWorld::GetSubsystem<UHoldoutCompositeSubsystem>(GetWorld());
 
-	UnregisterCompositeImpl();
+		if (ensureMsgf(IsValid(Subsystem), TEXT("Unexpected invalid holdout composite subsystem.")))
+		{
+			Subsystem->RegisterPrimitives(UE::HoldoutComposite::Private::GetPrimitives(CompositeMeshActors));
+		}
+	}
 }
 
-void AHoldoutComposite::Tick(float DeltaSeconds)
+void AHoldoutComposite::EndOfFrameUpdate(UWorld* InWorld)
 {
-	Super::Tick(DeltaSeconds);
+	SCOPE_CYCLE_COUNTER(STAT_HoldoutComposite_Actor_EndOfFrameUpdate);
+
+	if (!IsActive())
+	{
+		return;
+	}
+
+	// Only update materials once per frame
+	if (PreviousFrameCounter == GFrameCounter)
+	{
+		return;
+	}
+	else
+	{
+		PreviousFrameCounter = GFrameCounter;
+	}
+
+	if (bCompositeTextureProjection)
+	{
+		bool bUpdateViewProjection = false;
+
+		FMinimalViewInfo DesiredView;
+		FMatrix ViewMatrix, ProjectionMatrix, ViewProjectionMatrix;
+
+		TStrongObjectPtr<ACameraActor> CameraActor = ProjectionCamera.Pin();
+		if (CameraActor.IsValid())
+		{
+			if (UCameraComponent* CameraComponent = CameraActor->GetCameraComponent())
+			{
+				CameraComponent->GetCameraView(FApp::GetDeltaTime(), DesiredView);
+
+				UGameplayStatics::GetViewProjectionMatrix(DesiredView, ViewMatrix, ProjectionMatrix, ViewProjectionMatrix);
+				bUpdateViewProjection = true;
+			}
+		}
+		else if (const APlayerController* PlayerController = GetWorld()->GetFirstPlayerController())
+		{
+			if (AActor* ViewTarget = PlayerController->GetViewTarget())
+			{
+				UGameplayStatics::CalculateViewProjectionMatricesFromViewTarget(ViewTarget, ViewMatrix, ProjectionMatrix, ViewProjectionMatrix);
+				bUpdateViewProjection = true;
+			}
+		}
+#if WITH_EDITOR
+		else if (GCurrentLevelEditingViewportClient)
+		{
+			if (UCameraComponent* CameraComponent = GCurrentLevelEditingViewportClient->GetCameraComponentForView())
+			{
+				CameraComponent->GetCameraView(FApp::GetDeltaTime(), DesiredView);
+
+				UGameplayStatics::GetViewProjectionMatrix(DesiredView, ViewMatrix, ProjectionMatrix, ViewProjectionMatrix);
+				bUpdateViewProjection = true;
+			}
+			else
+			{
+				static FSceneViewFamily::ConstructionValues CVS(nullptr, nullptr, FEngineShowFlags(EShowFlagInitMode::ESFIM_Game));
+				CVS.SetTime(FGameTime::GetTimeSinceAppStart());
+
+				FSceneViewFamilyContext ViewFamily(CVS);
+				FSceneView* View = GCurrentLevelEditingViewportClient->CalcSceneView(&ViewFamily);
+
+				ViewProjectionMatrix = View->ViewMatrices.GetViewProjectionMatrix();
+				bUpdateViewProjection = true;
+			}
+		}
+#endif
+
+#if 0
+		if (CachedViewProjectionMatrix.Equals(ViewProjectionMatrix, UE_SMALL_NUMBER))
+		{
+			bUpdateViewProjection = false;
+		}
+#endif
 
-	constexpr bool bIsTick = true;
-	RegisterCompositeImpl(bIsTick);
+		if (bUpdateViewProjection)
+		{
+			SetMaterialParameterOnCompositeMeshes(
+				[this, &ViewProjectionMatrix](UMaterialInstanceDynamic* MID)
+				{
+					const FVector4 XPlane{ ViewProjectionMatrix.M[0][0], ViewProjectionMatrix.M[0][1], ViewProjectionMatrix.M[0][2], ViewProjectionMatrix.M[0][3] };
+					const FVector4 YPlane{ ViewProjectionMatrix.M[1][0], ViewProjectionMatrix.M[1][1], ViewProjectionMatrix.M[1][2], ViewProjectionMatrix.M[1][3] };
+					const FVector4 ZPlane{ ViewProjectionMatrix.M[2][0], ViewProjectionMatrix.M[2][1], ViewProjectionMatrix.M[2][2], ViewProjectionMatrix.M[2][3] };
+					const FVector4 WPlane{ ViewProjectionMatrix.M[3][0], ViewProjectionMatrix.M[3][1], ViewProjectionMatrix.M[3][2], ViewProjectionMatrix.M[3][3] };
+
+					MID->SetVectorParameterValue(TEXT("XPlane"), XPlane);
+					MID->SetVectorParameterValue(TEXT("YPlane"), YPlane);
+					MID->SetVectorParameterValue(TEXT("ZPlane"), ZPlane);
+					MID->SetVectorParameterValue(TEXT("WPlane"), WPlane);
+
+					// Also continuously update the composite texture, in case of user changes on composite meshes.
+					// TODO: Could be improved by only updating parameter upon dirty primitive deletages?
+					MID->SetTextureParameterValue(TEXT("CompositeTexture"), CompositeTexture.Get());
+				}
+			);
+
+			CachedViewProjectionMatrix = ViewProjectionMatrix;
+		}
+	}
 }
 
 void AHoldoutComposite::BeginDestroy()
@@ -131,11 +274,11 @@ void AHoldoutComposite::SetIsTemporarilyHiddenInEditor(bool bIsHidden)
 
 	if (IsTemporarilyHiddenInEditor())
 	{
-		UnregisterCompositeImpl();
+		RemoveAllState();
 	}
-	else
+	else if(IsActive())
 	{
-		RegisterCompositeImpl();
+		ApplyAllState();
 	}
 }
 bool AHoldoutComposite::ShouldTickIfViewportsOnly() const
@@ -153,63 +296,149 @@ void AHoldoutComposite::SetEnabled(bool bInEnabled)
 {
 	bIsEnabled = bInEnabled;
 
-	if (bIsEnabled)
+	if (IsActive())
 	{
-		RegisterCompositeImpl();
+		ApplyAllState();
 	}
 	else
 	{
-		UnregisterCompositeImpl();
+		RemoveAllState();
 	}
 }
 
-void AHoldoutComposite::RegisterCompositeImpl(bool bIsTick)
+UTexture* AHoldoutComposite::GetCompositeTexture() const
 {
-	using namespace UE::HoldoutComposite::Private;
+	return CompositeTexture.Get();
+}
 
-	if (GRegisterPrimitivesOnTick != bIsTick)
+void AHoldoutComposite::SetCompositeTexture(UTexture* InCompositeTexture)
+{
+	if (CompositeTexture != InCompositeTexture)
 	{
-		return;
+		CompositeTexture = InCompositeTexture;
+
+		if (IsActive())
+		{
+			ApplyAllState();
+		}
 	}
+}
 
-	if (!bIsEnabled)
+void AHoldoutComposite::SetEnableOverrideSettings(bool InEnable)
+{
+	if (bEnableOverrideSettings != InEnable)
 	{
-		return;
+		bEnableOverrideSettings = InEnable;
+
+		UHoldoutCompositeSubsystem* Subsystem = UWorld::GetSubsystem<UHoldoutCompositeSubsystem>(GetWorld());
+
+		if (ensureMsgf(IsValid(Subsystem), TEXT("Unexpected invalid holdout composite subsystem.")))
+		{
+			if (bEnableOverrideSettings)
+			{
+				// We leave the default holdout composite priority at 0, since only be one actor per level is allowed.
+				constexpr int32 DefaultPriority = 0;
+				Subsystem->OverrideCompositeSettings(OverrideSettings, DefaultPriority);
+			}
+			else
+			{
+				Subsystem->ResetCompositeSettings();
+			}
+		}
 	}
+}
+
+bool AHoldoutComposite::GetEnableOverrideSettings() const
+{
+	return bEnableOverrideSettings;
+}
+
+bool AHoldoutComposite::IsActive() const
+{
+	bool bIsActive = bIsEnabled;
 
 #if WITH_EDITOR
 	if (IsTemporarilyHiddenInEditor())
 	{
-		return;
+		bIsActive = false;
 	}
 #endif
 
+	return bIsActive;
+}
+
+void AHoldoutComposite::ApplyAllState()
+{
 	UHoldoutCompositeSubsystem* Subsystem = UWorld::GetSubsystem<UHoldoutCompositeSubsystem>(GetWorld());
+
 	if (IsValid(Subsystem))
 	{
-		if (bOverrideDefaultSettings)
+		Subsystem->RegisterPrimitives(UE::HoldoutComposite::Private::GetPrimitives(CompositeMeshActors));
+
+		if (!CompositePostProcessMaterials.IsEmpty())
+		{
+			Subsystem->SetCompositePostProcessMaterials(CompositePostProcessMaterials);
+		}
+
+		if (bEnableOverrideSettings)
 		{
 			// We leave the default holdout composite priority at 0, since only be one actor per level is allowed.
 			constexpr int32 DefaultPriority = 0;
-			Subsystem->OverrideCompositeSettings(SettingsOverride, DefaultPriority);
+			Subsystem->OverrideCompositeSettings(OverrideSettings, DefaultPriority);
+		}
+
+		if (bCompositeTextureOverrideRender)
+		{
+			Subsystem->SetCompositeTextureOverride(CompositeTexture);
 		}
-		Subsystem->RegisterPrimitives(GetPrimitives(CompositeMeshActors));
-		Subsystem->SetCompositeTextureOverride(bOverrideCompositeTexture ? CompositeTexture : nullptr);
-		Subsystem->SetCompositePostProcessMaterials(CompositePostProcessMaterials);
+
+		SetMaterialParameterOnCompositeMeshes(
+			[this](UMaterialInstanceDynamic* MID)
+			{
+				MID->SetTextureParameterValue(TEXT("CompositeTexture"), CompositeTexture.Get());
+			}
+		);
 	}
 }
 
-void AHoldoutComposite::UnregisterCompositeImpl()
+void AHoldoutComposite::RemoveAllState()
 {
-	using namespace UE::HoldoutComposite::Private;
-
 	UHoldoutCompositeSubsystem* Subsystem = UWorld::GetSubsystem<UHoldoutCompositeSubsystem>(GetWorld());
+	
 	if (IsValid(Subsystem))
 	{
+		// TODO: We probably want an unregister all method instead, since there can only be one actor anyway.
+		Subsystem->UnregisterPrimitives(UE::HoldoutComposite::Private::GetPrimitives(CompositeMeshActors));
 		Subsystem->ResetCompositePostProcessMaterials();
-		Subsystem->UnregisterPrimitives(GetPrimitives(CompositeMeshActors));
-		Subsystem->SetCompositeTextureOverride(nullptr);
 		Subsystem->ResetCompositeSettings();
+		Subsystem->SetCompositeTextureOverride(nullptr);
+	}
+}
+
+void AHoldoutComposite::SetMaterialParameterOnCompositeMeshes(TFunction<void(UMaterialInstanceDynamic* MID)> ApplyFn)
+{
+	using namespace UE::HoldoutComposite::Private;
+
+	TArray<UPrimitiveComponent*> Primitives = GetPrimitives(CompositeMeshActors);
+	
+	for (UPrimitiveComponent* Primitive : Primitives)
+	{
+		if (!IsValid(Primitive))
+		{
+			continue;
+		}
+
+		for (int32 ElementIndex = 0; ElementIndex < Primitive->GetNumMaterials(); ElementIndex++)
+		{
+			UMaterialInterface* Material = Primitive->GetMaterial(ElementIndex);
+			UMaterialInstanceDynamic* MID = Cast<UMaterialInstanceDynamic>(Material);
+			if (!MID)
+			{
+				MID = Primitive->CreateAndSetMaterialInstanceDynamicFromMaterial(ElementIndex, Material);
+			}
+
+			ApplyFn(MID);
+		}
 	}
 }
 
diff --git a/Engine/Plugins/Experimental/Compositing/HoldoutComposite/Source/HoldoutComposite/Public/HoldoutComposite.h b/Engine/Plugins/Experimental/Compositing/HoldoutComposite/Source/HoldoutComposite/Public/HoldoutComposite.h
--- a/Engine/Plugins/Experimental/Compositing/HoldoutComposite/Source/HoldoutComposite/Public/HoldoutComposite.h
+++ b/Engine/Plugins/Experimental/Compositing/HoldoutComposite/Source/HoldoutComposite/Public/HoldoutComposite.h
@@ -8,8 +8,11 @@
 
 #include "HoldoutComposite.generated.h"
 
+class ACameraActor;
+class UHoldoutCompositeSubsystem;
 class UPrimitiveComponent;
 class UMaterialInterface;
+class UMaterialInstanceDynamic;
 class UTexture;
 
 UCLASS()
@@ -18,7 +21,6 @@ class HOLDOUTCOMPOSITE_API AHoldoutComposite : public AActor
 	GENERATED_BODY()
 
 public:
-
 	AHoldoutComposite(const FObjectInitializer& ObjectInitializer);
 	~AHoldoutComposite();
 
@@ -37,52 +39,94 @@ public:
 	virtual bool ShouldTickIfViewportsOnly() const override;
 #endif // WITH_EDITOR
 
-	/* Get the enabled state of the component. */
+	/* Get the enabled state. */
 	UFUNCTION(BlueprintGetter)
 	bool IsEnabled() const;
 
-	/* Set the enabled state of the component. */
+	/* Set the enabled state. */
 	UFUNCTION(BlueprintSetter)
 	void SetEnabled(bool bInEnabled);
 
-private:
+	/* Get the composite texture. */
+	UFUNCTION(BlueprintGetter)
+	UTexture* GetCompositeTexture() const;
+
+	/* Set the composite texture. */
+	UFUNCTION(BlueprintSetter)
+	void SetCompositeTexture(UTexture* InCompositeTexture);
+
+	/* Set the override settings. */
+	UFUNCTION(BlueprintSetter)
+	void SetEnableOverrideSettings(bool InEnable);
+
+	/* Get the override settings. */
+	UFUNCTION(BlueprintGetter)
+	bool GetEnableOverrideSettings() const;
 
-	/* Whether or not the component activates the composite. */
-	UPROPERTY(EditAnywhere, BlueprintGetter = IsEnabled, BlueprintSetter = SetEnabled, Category = HoldoutComposite, meta = (AllowPrivateAccess = true))
+private:
+	/* Whether or not the holdout composite is active. */
+	UPROPERTY(EditAnywhere, BlueprintGetter = IsEnabled, BlueprintSetter = SetEnabled, Category = "HoldoutComposite", meta = (AllowPrivateAccess = true, DisplayPriority="0"))
 	bool bIsEnabled = true;
 
 public:
-
 	/* Actors whose primitive (mesh) components will be marked as primitive alpha holdout, rendered separately and composited with scene color after post-processing. */
-	UPROPERTY(EditAnywhere, BlueprintReadOnly, Category = HoldoutComposite)
+	UPROPERTY(EditAnywhere, BlueprintReadWrite, Category = "HoldoutComposite")
 	TArray<TObjectPtr<AActor>> CompositeMeshActors;
 
 	/** Post-process material(s) to be applied on the separate composited render.*/
-	UPROPERTY(EditAnywhere, BlueprintReadOnly, Category = HoldoutComposite)
+	UPROPERTY(EditAnywhere, BlueprintReadWrite, Category = "HoldoutComposite")
 	TArray<TObjectPtr<UMaterialInterface>> CompositePostProcessMaterials;
 
-	/** Enable composite texture override.*/
-	UPROPERTY(EditAnywhere, BlueprintReadOnly, Category = HoldoutComposite)
-	bool bOverrideCompositeTexture = false;
-
-	/* Instead of rendering the holdout composite actors separately, we only use them as primitive alpha holdout masks and composite the texture directly with scene color. */
-	UPROPERTY(EditAnywhere, BlueprintReadOnly, Category = HoldoutComposite, meta = (EditCondition = bOverrideCompositeTexture))
+private:
+	/* Texture to be used directly for alpha holdout composite and/or to be projected onto composite meshes. */
+	UPROPERTY(EditAnywhere, BlueprintGetter = GetCompositeTexture, BlueprintSetter = SetCompositeTexture, Category = "HoldoutComposite")
 	TObjectPtr<UTexture> CompositeTexture;
 
+public:
+	/** Directly sample the specified composite texture in screen space for the final post-processing composite, instead of using the separate built-in custom render pass. */
+	UPROPERTY(EditAnywhere, BlueprintReadWrite, Category = "HoldoutComposite", meta = (EditCondition = "CompositeTexture != nullptr"))
+	bool bCompositeTextureOverrideRender = false;
+
+	/** Activate projection of the composite texture onto the composite mesh actors. Continuously updates four material parameters called "XPlane", "YPlane", "ZPlane" & "WPlane" representing the view-projection matrix. */
+	UPROPERTY(EditAnywhere, BlueprintReadWrite, Category = "HoldoutComposite")
+	bool bCompositeTextureProjection = false;
+
+	/* Camera used for material projection. If none, the active viewport camera is used. */
+	UPROPERTY(EditAnywhere, BlueprintReadWrite, Category = "HoldoutComposite")
+	TWeakObjectPtr<ACameraActor> ProjectionCamera;
+
 	/** Enable local settings override.*/
-	UPROPERTY(EditAnywhere, BlueprintReadOnly, Category = HoldoutComposite)
-	bool bOverrideDefaultSettings = false;
+	UPROPERTY(EditAnywhere, BlueprintGetter = GetEnableOverrideSettings, BlueprintSetter = SetEnableOverrideSettings, Category = "HoldoutComposite")
+	bool bEnableOverrideSettings = false;
 
 	/** Post process settings to use for this holdout composite. */
-	UPROPERTY(EditAnywhere, BlueprintReadOnly, Category = HoldoutComposite, meta=(EditCondition=bOverrideDefaultSettings))
-	FHoldoutCompositeSettings SettingsOverride = {};
+	UPROPERTY(EditAnywhere, BlueprintReadOnly, Category = "HoldoutComposite", meta = (EditCondition = bOverrideDefaultSettings, DisplayAfter = "bOverrideDefaultSettings"))
+	FHoldoutCompositeSettings OverrideSettings = {};
 
 private:
 
-	/* Private implementation of the register method. */
-	void RegisterCompositeImpl(bool bIsTick = false);
+	/* Returns true if the holdout composite is active. */
+	bool IsActive() const;
+
+	/* Apply all actor state to the holdout composite subsystem. */
+	void ApplyAllState();
+
+	/* Remove all actor state from the holdout composite subsystem. */
+	void RemoveAllState();
+
+	/* End-of-frame update, primarily used for updating MID parameters on all composite mesh primitives (after viewport ticks). */
+	void EndOfFrameUpdate(UWorld* InWorld);
+
+	/* Convenience function to visit all composite meshes and update parameters on their MIDs. */
+	void SetMaterialParameterOnCompositeMeshes(TFunction<void(UMaterialInstanceDynamic* MID)> ApplyFn);
+
+	/* Cached view projection matrix to avoid unecessary reprojections. */
+	FMatrix CachedViewProjectionMatrix = FMatrix::Identity;
+
+	/* Cached previous frame counter to avoid double updates on the same frame. */
+	uint64 PreviousFrameCounter = 0;
 
-	/* Private implementation of the unregister method. */
-	void UnregisterCompositeImpl();
+	/* Post-viewport update callback. */
+	FDelegateHandle OnWorldPreSendAllEndOfFrameUpdatesHandle;
 };
 
diff --git a/Engine/Plugins/Experimental/Compositing/HoldoutComposite/Source/HoldoutComposite/Public/HoldoutCompositeModule.h b/Engine/Plugins/Experimental/Compositing/HoldoutComposite/Source/HoldoutComposite/Public/HoldoutCompositeModule.h
--- a/Engine/Plugins/Experimental/Compositing/HoldoutComposite/Source/HoldoutComposite/Public/HoldoutCompositeModule.h
+++ b/Engine/Plugins/Experimental/Compositing/HoldoutComposite/Source/HoldoutComposite/Public/HoldoutCompositeModule.h
@@ -4,9 +4,12 @@
 
 #include "Modules/ModuleInterface.h"
 #include "Logging/LogMacros.h"
+#include "Stats/Stats.h"
 
 HOLDOUTCOMPOSITE_API DECLARE_LOG_CATEGORY_EXTERN(LogHoldoutComposite, Log, All);
 
+DECLARE_STATS_GROUP(TEXT("HoldoutComposite"), STATGROUP_HoldoutComposite, STATCAT_Advanced);
+
 class FHoldoutCompositeModule : public IModuleInterface
 {
 public:
