HoldoutComposite: Fix for potential crash on invalid world during subsystem initialization.

diff --git a/Engine/Plugins/Experimental/Compositing/HoldoutComposite/Source/HoldoutComposite/Private/HoldoutCompositeSubsystem.cpp b/Engine/Plugins/Experimental/Compositing/HoldoutComposite/Source/HoldoutComposite/Private/HoldoutCompositeSubsystem.cpp
--- a/Engine/Plugins/Experimental/Compositing/HoldoutComposite/Source/HoldoutComposite/Private/HoldoutCompositeSubsystem.cpp
+++ b/Engine/Plugins/Experimental/Compositing/HoldoutComposite/Source/HoldoutComposite/Private/HoldoutCompositeSubsystem.cpp
@@ -52,9 +52,11 @@ void UHoldoutCompositeSubsystem::Initialize(FSubsystemCollectionBase& Collection
 	Super::Initialize(Collection);
 
 	UWorld* World = GetWorld();
-	check(IsValid(World));
 
-	HoldoutCompositeViewExtension = FSceneViewExtensions::NewExtension<FHoldoutCompositeSceneViewExtension>(World);
+	if (IsValid(World))
+	{
+		HoldoutCompositeViewExtension = FSceneViewExtensions::NewExtension<FHoldoutCompositeSceneViewExtension>(World);
+	}
 }
 
 void UHoldoutCompositeSubsystem::Deinitialize()
