HoldoutComposite: Major plugin refactor into a modular system of render passes.

* Breaking change within 5.6, but most of the existing functionality is preserved. Support for "follow scene exposure" setting was removed: it should be applied by users in materials, idem for "CompositeTexture" material paramater update.
* The actor is now in control of which passes are done during post-processing. When only registration components are used (no actor), we issue default render work which matches previous behavior.
* SSR input & after tonemap passes are now sharing a single shader, which correctly handles alpha inversion, TSR distortion, etc, based on its connected inputs.

diff --git a/Engine/Plugins/Experimental/Compositing/HoldoutComposite/Config/BaseHoldoutComposite.ini b/Engine/Plugins/Experimental/Compositing/HoldoutComposite/Config/BaseHoldoutComposite.ini
deleted file mode 100644
diff --git a/Engine/Plugins/Experimental/Compositing/HoldoutComposite/Shaders/Private/HoldoutCompositeFinal.usf b/Engine/Plugins/Experimental/Compositing/HoldoutComposite/Shaders/Private/HoldoutCompositeFinal.usf
deleted file mode 100644
diff --git a/Engine/Plugins/Experimental/Compositing/HoldoutComposite/Shaders/Private/HoldoutCompositeMerge.usf b/Engine/Plugins/Experimental/Compositing/HoldoutComposite/Shaders/Private/HoldoutCompositeMerge.usf
new file mode 100644
--- /dev/null
+++ b/Engine/Plugins/Experimental/Compositing/HoldoutComposite/Shaders/Private/HoldoutCompositeMerge.usf
@@ -0,0 +1,162 @@
+// Copyright Epic Games, Inc. All Rights Reserved.
+
+#include "/Engine/Private/Common.ush"
+#include "/Engine/Private/GammaCorrectionCommon.ush"
+#include "/Engine/Private/ScreenPass.ush"
+#include "/Engine/Private/LensDistortion.ush"
+
+#define ESourceEncoding_None  0
+#define ESourceEncoding_Gamma 1
+#define ESourceEncoding_sRGB  2
+
+#define EDistortionUV_None        0
+#define EDistortionUV_Distorted   1
+#define EDistortionUV_Undistorted 2
+
+#define EHoldoutCompositePassMergeOp_Over 0
+#define EHoldoutCompositePassMergeOp_Add  1
+
+SCREEN_PASS_TEXTURE_VIEWPORT(Input0)
+SCREEN_PASS_TEXTURE_VIEWPORT(Input1)
+SCREEN_PASS_TEXTURE_VIEWPORT(Output)
+
+#define INPUT_PASS_TEXTURE_MEMBER(StructName, MemberType, MemberName) MemberType StructName##_##MemberName;
+
+#define INPUT_PASS_TEXTURE(StructName) \
+	INPUT_PASS_TEXTURE_MEMBER(StructName, Texture2D, Texture) \
+	INPUT_PASS_TEXTURE_MEMBER(StructName, SamplerState, Sampler) \
+    INPUT_PASS_TEXTURE_MEMBER(StructName, uint, bAlphaIsOpacity) \
+    INPUT_PASS_TEXTURE_MEMBER(StructName, uint, SourceEncoding) \
+    INPUT_PASS_TEXTURE_MEMBER(StructName, uint, DistortionUV)
+    
+
+INPUT_PASS_TEXTURE(Tex0);
+INPUT_PASS_TEXTURE(Tex1);
+
+Texture2D<float2> DistortingDisplacementTexture;
+SamplerState DistortingDisplacementSampler;
+Texture2D<float2> UndistortingDisplacementTexture;
+SamplerState UndistortingDisplacementSampler;
+
+float2 DisplayGamma; // DisplayGamma, 1/DisplayGamma
+uint MergeOp;
+uint OutputEncoding;
+uint bFlags; // 0: bOutputSceneColor, 1: bUnpremulInput1
+
+float SafeRcp(float x)
+{
+	return x > 0.0 ? rcp(x) : 0.0;
+}
+
+float4 Unpremultiply(float4 InColor)
+{
+	return float4(InColor.xyz * SafeRcp(InColor.w), InColor.w);
+}
+
+float3 Decode(float3 InColor, uint InEncoding)
+{
+    // TODO: HDR cases will need to be handled explicitely.
+    BRANCH
+    if(InEncoding == ESourceEncoding_Gamma)
+    {
+        return pow(InColor, DisplayGamma.x);
+    }
+    else if(InEncoding == ESourceEncoding_sRGB)
+    {
+        return sRGBToLinear(InColor);
+    }
+    else
+    {
+        return InColor;
+    }
+}
+
+float3 Encode(float3 InColor, uint InEncoding)
+{
+    BRANCH
+    if(InEncoding == ESourceEncoding_Gamma)
+    {
+        return pow(InColor, DisplayGamma.y);
+    }
+    else if(InEncoding == ESourceEncoding_sRGB)
+    {
+        return LinearToSrgb(InColor);
+    }
+    else
+    {
+        return InColor;
+    }
+}
+
+float2 CalcInputUVs(float2 ViewportUV, float2 ViewportSize, float2 ViewportMin, float2 ExtentInverse, uint DistortionUVType)
+{
+#if USE_DISTORTION
+    BRANCH
+    if(DistortionUVType == EDistortionUV_Distorted)
+    {
+        ViewportUV = ApplyLensDistortionOnViewportUV(DistortingDisplacementTexture, DistortingDisplacementSampler, ViewportUV);
+    }
+    else if(DistortionUVType == EDistortionUV_Undistorted)
+    {
+        ViewportUV = ApplyLensDistortionOnViewportUV(UndistortingDisplacementTexture, UndistortingDisplacementSampler, ViewportUV);
+    }
+#endif
+
+    return (ViewportUV * ViewportSize + ViewportMin) * ExtentInverse;
+}
+
+float GetOpacity(float InAlpha, uint bAlphaIsOpacity)
+{
+	return select(bAlphaIsOpacity > 0, InAlpha, 1.0 - InAlpha);
+}
+
+void MainPS(
+    float4 SvPosition : SV_POSITION,
+    out float4 OutColor : SV_Target0
+    )
+{
+    float2 ViewportUV = (SvPosition.xy - Output_ViewportMin) * Output_ViewportSizeInverse;
+
+    // Calculate texture UVs, with/without distortion
+    float2 Tex0_UV = CalcInputUVs(ViewportUV, Input0_ViewportSize, Input0_ViewportMin, Input0_ExtentInverse, Tex0_DistortionUV);
+    float2 Tex1_UV = CalcInputUVs(ViewportUV, Input1_ViewportSize, Input1_ViewportMin, Input1_ExtentInverse, Tex1_DistortionUV);
+
+    // Sample textures
+    float4 Color0 = Tex0_Texture.SampleLevel(Tex0_Sampler, Tex0_UV, 0);
+    float4 Color1 = Tex1_Texture.SampleLevel(Tex1_Sampler, Tex1_UV, 0);
+
+    // Linearize source colors
+    Color0.rgb = Decode(Color0.rgb, Tex0_SourceEncoding);
+    Color1.rgb = Decode(Color1.rgb, Tex1_SourceEncoding);
+
+    // Normalize alpha as opacity
+    Color0.a = GetOpacity(Color0.a, Tex0_bAlphaIsOpacity);
+    Color1.a = GetOpacity(Color1.a, Tex1_bAlphaIsOpacity);
+
+    {
+        // NOTE: This is necessary to extract the solid colors behind translucent alpha holdout masks.
+        // Could be moved elsewhere as a separate pre-processing step, with minor added cost.
+        const bool bUnpremulInput1 = ((bFlags >> 1) & 1);
+        Color1 = select(bUnpremulInput1, Unpremultiply(Color1), Color1);
+    }
+    
+    if(MergeOp == EHoldoutCompositePassMergeOp_Over)
+    {
+        OutColor = Color0 + Color1 * (1.0 - Color0.a);
+    }
+    else if(MergeOp == EHoldoutCompositePassMergeOp_Add)
+    {
+        OutColor = Color0 + Color1;
+    }
+    else
+    {
+        OutColor = Color0;
+    }
+
+    // Re-encode color
+    OutColor.rgb = Encode(OutColor.rgb, OutputEncoding);
+
+    // Optionally re-invert alpha
+    const bool bOutputSceneColor = (bFlags & 1);
+    OutColor.a = saturate((bOutputSceneColor > 0, 1.0 - OutColor.a, OutColor.a));
+}
diff --git a/Engine/Plugins/Experimental/Compositing/HoldoutComposite/Shaders/Private/HoldoutCompositeSSRInput.usf b/Engine/Plugins/Experimental/Compositing/HoldoutComposite/Shaders/Private/HoldoutCompositeSSRInput.usf
deleted file mode 100644
diff --git a/Engine/Plugins/Experimental/Compositing/HoldoutComposite/Source/HoldoutComposite/Private/HoldoutComposite.cpp b/Engine/Plugins/Experimental/Compositing/HoldoutComposite/Source/HoldoutComposite/Private/HoldoutComposite.cpp
--- a/Engine/Plugins/Experimental/Compositing/HoldoutComposite/Source/HoldoutComposite/Private/HoldoutComposite.cpp
+++ b/Engine/Plugins/Experimental/Compositing/HoldoutComposite/Source/HoldoutComposite/Private/HoldoutComposite.cpp
@@ -5,6 +5,8 @@
 #include "HoldoutCompositeModule.h"
 #include "HoldoutCompositeSettings.h"
 
+#include "Passes/HoldoutCompositePassMerge.h"
+
 #include "Components/PrimitiveComponent.h"
 #include "Camera/CameraActor.h"
 #include "Camera/CameraComponent.h"
@@ -72,6 +74,10 @@ AHoldoutComposite::AHoldoutComposite(const FObjectInitializer& ObjectInitializer
 	}
 
 	SetActorEnableCollision(false);
+
+	// Note: Default setup here should match UE::HoldoutComposite::Private::FRenderWork::GetDefault();
+	Passes_SSRInput.Add(CreateDefaultSubobject<UHoldoutCompositePassMerge>(TEXT("MergePass_SSRInput")));
+	Passes_AfterToneMap.Add(CreateDefaultSubobject<UHoldoutCompositePassMerge>(TEXT("MergePass_AfterTonemap")));
 }
 
 AHoldoutComposite::~AHoldoutComposite() = default;
@@ -79,25 +85,6 @@ AHoldoutComposite::~AHoldoutComposite() = default;
 void AHoldoutComposite::PostActorCreated()
 {
 	Super::PostActorCreated();
-
-#if WITH_EDITOR
-	int32 HoldoutCompositeCount = 0;
-	for (TObjectIterator<AHoldoutComposite> It; It; ++It)
-	{
-		HoldoutCompositeCount++;
-	}
-
-	if (HoldoutCompositeCount > 1)
-	{
-		if (GEngine)
-		{
-			GEngine->AddOnScreenDebugMessage(-1, 5.0f, FColor::Red, TEXT("No more than one holdout composite actor is allowed."));
-		}
-
-		UE_LOG(LogHoldoutComposite, Error, TEXT("No more than one holdout composite actor is allowed, disabling. Please remove."));
-		bIsEnabled = false;
-	}
-#endif // WITH_EDITOR
 }
 
 void AHoldoutComposite::PostRegisterAllComponents()
@@ -120,12 +107,23 @@ void AHoldoutComposite::PostRegisterAllComponents()
 	*/
 	OnWorldPreSendAllEndOfFrameUpdatesHandle = FWorldDelegates::OnWorldPreSendAllEndOfFrameUpdates.AddUObject(this, &AHoldoutComposite::EndOfFrameUpdate);
 
+	UHoldoutCompositeSubsystem* Subsystem = UWorld::GetSubsystem<UHoldoutCompositeSubsystem>(GetWorld());
+	if (IsValid(Subsystem))
+	{
+		Subsystem->RegisterActor(this);
+	}
 }
 
 void AHoldoutComposite::PostUnregisterAllComponents()
 {
 	using namespace UE::HoldoutComposite::Private;
 
+	UHoldoutCompositeSubsystem* Subsystem = UWorld::GetSubsystem<UHoldoutCompositeSubsystem>(GetWorld());
+	if (IsValid(Subsystem))
+	{
+		Subsystem->ResetActor();
+	}
+
 	FWorldDelegates::OnWorldPreSendAllEndOfFrameUpdates.Remove(OnWorldPreSendAllEndOfFrameUpdatesHandle);
 
 	RemoveAllState();
@@ -139,10 +137,14 @@ void AHoldoutComposite::Tick(float DeltaTime)
 
 	SCOPE_CYCLE_COUNTER(STAT_HoldoutComposite_Actor_Tick);
 
-	if (IsActive() && UE::HoldoutComposite::Private::GRegisterPrimitivesOnTick)
+	if (!IsActive())
 	{
-		UHoldoutCompositeSubsystem* Subsystem = UWorld::GetSubsystem<UHoldoutCompositeSubsystem>(GetWorld());
+		return;
+	}
 
+	if (UE::HoldoutComposite::Private::GRegisterPrimitivesOnTick)
+	{
+		UHoldoutCompositeSubsystem* Subsystem = UWorld::GetSubsystem<UHoldoutCompositeSubsystem>(GetWorld());
 		if (ensureMsgf(IsValid(Subsystem), TEXT("Unexpected invalid holdout composite subsystem.")))
 		{
 			Subsystem->RegisterPrimitives(UE::HoldoutComposite::Private::GetPrimitives(CompositeMeshActors));
@@ -240,10 +242,6 @@ void AHoldoutComposite::EndOfFrameUpdate(UWorld* InWorld)
 					MID->SetVectorParameterValue(TEXT("YPlane"), YPlane);
 					MID->SetVectorParameterValue(TEXT("ZPlane"), ZPlane);
 					MID->SetVectorParameterValue(TEXT("WPlane"), WPlane);
-
-					// Also continuously update the composite texture, in case of user changes on composite meshes.
-					// TODO: Could be improved by only updating parameter upon dirty primitive deletages?
-					MID->SetTextureParameterValue(TEXT("CompositeTexture"), CompositeTexture.Get());
 				}
 			);
 
@@ -263,9 +261,27 @@ void AHoldoutComposite::PreEditChange(FProperty* PropertyThatWillChange)
 	Super::PreEditChange(PropertyThatWillChange);
 }
 
-void AHoldoutComposite::PostEditChangeProperty(FPropertyChangedEvent& PropertyChangedEvent)
+void AHoldoutComposite::PostEditChangeChainProperty(struct FPropertyChangedChainEvent& PropertyChangedEvent)
 {
-	Super::PostEditChangeProperty(PropertyChangedEvent);
+	Super::PostEditChangeChainProperty(PropertyChangedEvent);
+
+	const FName BasePropertyName = (PropertyChangedEvent.MemberProperty ? PropertyChangedEvent.MemberProperty->GetFName() : NAME_None);
+	const FName PropertyName = (PropertyChangedEvent.Property ? PropertyChangedEvent.Property->GetFName() : NAME_None);
+
+	if (PropertyName == GET_MEMBER_NAME_STRING_CHECKED(FHoldoutCompositeProcessingPasses, Passes))
+	{
+		for (FHoldoutCompositeProcessingPasses& Processing : CompositeProcessing)
+		{
+			for (TObjectPtr<UHoldoutCompositePassBase>& Pass : Processing.Passes)
+			{
+				if (IsValid(Pass))
+				{
+					// Disallow input change on input processing passes...
+					Pass->SetCanEditInput(false);
+				}
+			}
+		}
+	}
 }
 
 void AHoldoutComposite::SetIsTemporarilyHiddenInEditor(bool bIsHidden)
@@ -324,35 +340,6 @@ void AHoldoutComposite::SetCompositeTexture(UTexture* InCompositeTexture)
 	}
 }
 
-void AHoldoutComposite::SetEnableOverrideSettings(bool InEnable)
-{
-	if (bEnableOverrideSettings != InEnable)
-	{
-		bEnableOverrideSettings = InEnable;
-
-		UHoldoutCompositeSubsystem* Subsystem = UWorld::GetSubsystem<UHoldoutCompositeSubsystem>(GetWorld());
-
-		if (ensureMsgf(IsValid(Subsystem), TEXT("Unexpected invalid holdout composite subsystem.")))
-		{
-			if (bEnableOverrideSettings)
-			{
-				// We leave the default holdout composite priority at 0, since only be one actor per level is allowed.
-				constexpr int32 DefaultPriority = 0;
-				Subsystem->OverrideCompositeSettings(OverrideSettings, DefaultPriority);
-			}
-			else
-			{
-				Subsystem->ResetCompositeSettings();
-			}
-		}
-	}
-}
-
-bool AHoldoutComposite::GetEnableOverrideSettings() const
-{
-	return bEnableOverrideSettings;
-}
-
 bool AHoldoutComposite::IsActive() const
 {
 	bool bIsActive = bIsEnabled;
@@ -375,29 +362,7 @@ void AHoldoutComposite::ApplyAllState()
 	{
 		Subsystem->RegisterPrimitives(UE::HoldoutComposite::Private::GetPrimitives(CompositeMeshActors));
 
-		if (!CompositePostProcessMaterials.IsEmpty())
-		{
-			Subsystem->SetCompositePostProcessMaterials(CompositePostProcessMaterials);
-		}
-
-		if (bEnableOverrideSettings)
-		{
-			// We leave the default holdout composite priority at 0, since only be one actor per level is allowed.
-			constexpr int32 DefaultPriority = 0;
-			Subsystem->OverrideCompositeSettings(OverrideSettings, DefaultPriority);
-		}
-
-		if (bCompositeTextureOverrideRender)
-		{
-			Subsystem->SetCompositeTextureOverride(CompositeTexture);
-		}
-
-		SetMaterialParameterOnCompositeMeshes(
-			[this](UMaterialInstanceDynamic* MID)
-			{
-				MID->SetTextureParameterValue(TEXT("CompositeTexture"), CompositeTexture.Get());
-			}
-		);
+		Subsystem->SetCompositeTextureOverride(CompositeTexture);
 	}
 }
 
@@ -409,8 +374,7 @@ void AHoldoutComposite::RemoveAllState()
 	{
 		// TODO: We probably want an unregister all method instead, since there can only be one actor anyway.
 		Subsystem->UnregisterPrimitives(UE::HoldoutComposite::Private::GetPrimitives(CompositeMeshActors));
-		Subsystem->ResetCompositePostProcessMaterials();
-		Subsystem->ResetCompositeSettings();
+		
 		Subsystem->SetCompositeTextureOverride(nullptr);
 	}
 }
@@ -441,4 +405,3 @@ void AHoldoutComposite::SetMaterialParameterOnCompositeMeshes(TFunction<void(UMa
 		}
 	}
 }
-
diff --git a/Engine/Plugins/Experimental/Compositing/HoldoutComposite/Source/HoldoutComposite/Private/HoldoutCompositeSceneViewExtension.cpp b/Engine/Plugins/Experimental/Compositing/HoldoutComposite/Source/HoldoutComposite/Private/HoldoutCompositeSceneViewExtension.cpp
--- a/Engine/Plugins/Experimental/Compositing/HoldoutComposite/Source/HoldoutComposite/Private/HoldoutCompositeSceneViewExtension.cpp
+++ b/Engine/Plugins/Experimental/Compositing/HoldoutComposite/Source/HoldoutComposite/Private/HoldoutCompositeSceneViewExtension.cpp
@@ -5,6 +5,7 @@
 #include "HoldoutCompositeModule.h"
 #include "HoldoutComposite.h"
 
+#include "CommonRenderResources.h"
 #include "Components/PrimitiveComponent.h"
 #include "Containers/Map.h"
 #include "Containers/Set.h"
@@ -15,7 +16,6 @@
 #include "Materials/Material.h"
 #include "PixelShaderUtils.h"
 #include "PostProcess/LensDistortion.h"
-#include "PostProcess/PostProcessAA.h"
 #include "PostProcess/PostProcessMaterialInputs.h"
 #include "Rendering/CustomRenderPass.h"
 #include "RenderGraphBuilder.h"
@@ -26,48 +26,6 @@
 #include "TextureResource.h"
 
 DECLARE_GPU_STAT_NAMED(FHoldoutCompositeDilate, TEXT("HoldoutComposite.Dilate"));
-DECLARE_GPU_STAT_NAMED(FHoldoutCompositeGamma, TEXT("HoldoutComposite.Gamma"));
-DECLARE_GPU_STAT_NAMED(FHoldoutCompositeSSRInput, TEXT("HoldoutComposite.SSRInput"));
-DECLARE_GPU_STAT_NAMED(FHoldoutCompositeFinal, TEXT("HoldoutComposite.Final"));
-
-namespace HoldoutComposite
-{
-	enum class ESourceEncoding : uint32
-	{
-		Linear = 0,
-		Gamma = 1,
-		sRGB = 2,
-	};
-
-	static constexpr float DefaultGamma = 2.2f;
-
-	// Return the RDG texture from the the override texture uobject
-	FRDGTextureRef GetTextureOverride_RenderThread(FRDGBuilder& GraphBuilder, const TWeakObjectPtr<UTexture> InTexture)
-	{
-		constexpr bool bEvenIfPendingKill = false;
-		constexpr bool bThreadSafeTest = true;
-		if (InTexture.IsValid(bEvenIfPendingKill, bThreadSafeTest))
-		{
-			if (FTextureResource* TextureResource = InTexture.GetEvenIfUnreachable()->GetResource())
-			{
-				const FTextureRHIRef& TextureRHI = TextureResource->GetTextureRHI();
-				if (TextureRHI.IsValid())
-				{
-					FRDGTextureRef TextureRDG = GraphBuilder.FindExternalTexture(TextureRHI.GetReference());
-
-					if (TextureRDG == nullptr)
-					{
-						TextureRDG = GraphBuilder.RegisterExternalTexture(CreateRenderTarget(TextureRHI.GetReference(), TEXT("CompositeTextureOverride")));
-					}
-
-					return TextureRDG;
-				}
-			}
-		}
-
-		return nullptr;
-	}
-}
 
 class FHoldoutCompositeDilateShader : public FGlobalShader
 {
@@ -91,120 +49,53 @@ class FHoldoutCompositeDilateShader : public FGlobalShader
 };
 IMPLEMENT_GLOBAL_SHADER(FHoldoutCompositeDilateShader, "/Plugin/HoldoutComposite/Private/HoldoutCompositeDilate.usf", "MainCS", SF_Compute);
 
-class FHoldoutCompositeGammaShader : public FGlobalShader
-{
-	DECLARE_GLOBAL_SHADER(FHoldoutCompositeGammaShader);
-	SHADER_USE_PARAMETER_STRUCT(FHoldoutCompositeGammaShader, FGlobalShader);
-
-	BEGIN_SHADER_PARAMETER_STRUCT(FParameters, )
-		SHADER_PARAMETER_RDG_TEXTURE(Texture2D, InputTexture)
-		SHADER_PARAMETER_SAMPLER(SamplerState, InputSampler)
-		SHADER_PARAMETER(FIntPoint, Dimensions)
-		SHADER_PARAMETER(float, GammaPower)
-		RENDER_TARGET_BINDING_SLOTS()
-	END_SHADER_PARAMETER_STRUCT()
-};
-IMPLEMENT_GLOBAL_SHADER(FHoldoutCompositeGammaShader, "/Plugin/HoldoutComposite/Private/HoldoutCompositeGamma.usf", "MainPS", SF_Pixel);
-
-BEGIN_SHADER_PARAMETER_STRUCT(FHoldoutCompositeCommonParameters, )
-	SHADER_PARAMETER_STRUCT(FScreenPassTextureViewportParameters, Input)
-	SHADER_PARAMETER_STRUCT(FScreenPassTextureViewportParameters, Composite)
-	SHADER_PARAMETER_STRUCT(FScreenPassTextureViewportParameters, Output)
-	SHADER_PARAMETER_RDG_TEXTURE(Texture2D, InputTexture)
-	SHADER_PARAMETER_SAMPLER(SamplerState, InputSampler)
-	SHADER_PARAMETER_RDG_TEXTURE(Texture2D, CompositeTexture)
-	SHADER_PARAMETER_SAMPLER(SamplerState, CompositeSampler)
-	SHADER_PARAMETER_RDG_TEXTURE(Texture2D, DistortingDisplacementTexture)
-	SHADER_PARAMETER_SAMPLER(SamplerState, DistortingDisplacementSampler)
-	SHADER_PARAMETER_RDG_TEXTURE(Texture2D, UndistortingDisplacementTexture)
-	SHADER_PARAMETER_SAMPLER(SamplerState, UndistortingDisplacementSampler)
-END_SHADER_PARAMETER_STRUCT()
-
-class FHoldoutCompositeSSRInputShader : public FGlobalShader
-{
-	DECLARE_GLOBAL_SHADER(FHoldoutCompositeSSRInputShader);
-	SHADER_USE_PARAMETER_STRUCT(FHoldoutCompositeSSRInputShader, FGlobalShader);
-
-	BEGIN_SHADER_PARAMETER_STRUCT(FParameters, )
-		SHADER_PARAMETER_STRUCT_INCLUDE(FHoldoutCompositeCommonParameters, Common)
-		SHADER_PARAMETER(float, LastGlobalExposure)
-		RENDER_TARGET_BINDING_SLOTS()
-	END_SHADER_PARAMETER_STRUCT()
-};
-IMPLEMENT_GLOBAL_SHADER(FHoldoutCompositeSSRInputShader, "/Plugin/HoldoutComposite/Private/HoldoutCompositeSSRInput.usf", "MainPS", SF_Pixel);
-
-class FHoldoutCompositeFinalShader : public FGlobalShader
+namespace UE
 {
-	DECLARE_GLOBAL_SHADER(FHoldoutCompositeFinalShader);
-	SHADER_USE_PARAMETER_STRUCT(FHoldoutCompositeFinalShader, FGlobalShader);
-
-	class FUseGlobalExposure : SHADER_PERMUTATION_BOOL("USE_GLOBAL_EXPOSURE");
-	using FPermutationDomain = TShaderPermutationDomain<FUseGlobalExposure>;
+	namespace HoldoutComposite
+	{
+		namespace Private
+		{
+			void AddDilatePass(FRDGBuilder& GraphBuilder, FRDGTextureRef Input, FRDGTextureRef Output, ERHIFeatureLevel::Type FeatureLevel)
+			{
+				RDG_EVENT_SCOPE_STAT(GraphBuilder, FHoldoutCompositeDilate, "HoldoutComposite.Dilate");
+				RDG_GPU_STAT_SCOPE(GraphBuilder, FHoldoutCompositeDilate);
+
+				FGlobalShaderMap* GlobalShaderMap = GetGlobalShaderMap(FeatureLevel);
+				const FIntPoint TextureSize = Input->Desc.Extent;
+
+				FHoldoutCompositeDilateShader::FParameters* PassParameters = GraphBuilder.AllocParameters<FHoldoutCompositeDilateShader::FParameters>();
+				PassParameters->InputTexture = Input;
+				PassParameters->RWOutputTexture = GraphBuilder.CreateUAV(Output);
+				PassParameters->Dimensions = TextureSize;
+
+				TShaderMapRef<FHoldoutCompositeDilateShader> ComputeShader(GlobalShaderMap);
+				FComputeShaderUtils::AddPass(
+					GraphBuilder,
+					RDG_EVENT_NAME("HoldoutComposite.Dilate (%dx%d)", TextureSize.X, TextureSize.Y),
+					GSupportsEfficientAsyncCompute ? ERDGPassFlags::AsyncCompute : ERDGPassFlags::Compute,
+					ComputeShader,
+					PassParameters,
+					FComputeShaderUtils::GetGroupCount(TextureSize, FHoldoutCompositeDilateShader::ThreadGroupSize)
+				);
+			}
 
-	BEGIN_SHADER_PARAMETER_STRUCT(FParameters, )
-		SHADER_PARAMETER_STRUCT_REF(FViewUniformShaderParameters, View)
-		SHADER_PARAMETER_STRUCT_INCLUDE(FHoldoutCompositeCommonParameters, Common)
-		SHADER_PARAMETER_RDG_BUFFER_SRV(StructuredBuffer<float4>, EyeAdaptationBuffer)
-		SHADER_PARAMETER(FUint32Vector, Encodings)
-		SHADER_PARAMETER(FVector3f, DisplayGamma)
-		SHADER_PARAMETER(uint32, bApplyTSRLensDistortion)
-		SHADER_PARAMETER(uint32, bIsCustomRenderPassTexture)
-		RENDER_TARGET_BINDING_SLOTS()
-	END_SHADER_PARAMETER_STRUCT()
-};
-IMPLEMENT_GLOBAL_SHADER(FHoldoutCompositeFinalShader, "/Plugin/HoldoutComposite/Private/HoldoutCompositeFinal.usf", "MainPS", SF_Pixel);
+			FRDGTextureDesc GetPostProcessingDesc(FRDGTextureDesc InDesc, EPixelFormat InOutputFormat)
+			{
+				FRDGTextureDesc OutputDesc = InDesc;
+				OutputDesc.Dimension = ETextureDimension::Texture2D;
+				OutputDesc.ArraySize = 1;
 
-namespace HoldoutComposite
-{
-	void AddDilatePass(FRDGBuilder& GraphBuilder, FRDGTextureRef Input, FRDGTextureRef Output, ERHIFeatureLevel::Type FeatureLevel)
-	{
-		RDG_EVENT_SCOPE_STAT(GraphBuilder, FHoldoutCompositeDilate, "HoldoutComposite.Dilate");
-		RDG_GPU_STAT_SCOPE(GraphBuilder, FHoldoutCompositeDilate);
-
-		FGlobalShaderMap* GlobalShaderMap = GetGlobalShaderMap(FeatureLevel);
-		const FIntPoint TextureSize = Input->Desc.Extent;
-
-		FHoldoutCompositeDilateShader::FParameters* PassParameters = GraphBuilder.AllocParameters<FHoldoutCompositeDilateShader::FParameters>();
-		PassParameters->InputTexture = Input;
-		PassParameters->RWOutputTexture = GraphBuilder.CreateUAV(Output);
-		PassParameters->Dimensions = TextureSize;
-
-		TShaderMapRef<FHoldoutCompositeDilateShader> ComputeShader(GlobalShaderMap);
-		FComputeShaderUtils::AddPass(
-			GraphBuilder,
-			RDG_EVENT_NAME("HoldoutComposite.Dilate (%dx%d)", TextureSize.X, TextureSize.Y),
-			GSupportsEfficientAsyncCompute ? ERDGPassFlags::AsyncCompute : ERDGPassFlags::Compute,
-			ComputeShader,
-			PassParameters,
-			FComputeShaderUtils::GetGroupCount(TextureSize, FHoldoutCompositeDilateShader::ThreadGroupSize)
-		);
-	}
+				OutputDesc.Reset();
+				if (InOutputFormat != PF_Unknown)
+				{
+					OutputDesc.Format = InOutputFormat;
+				}
+				OutputDesc.ClearValue = FClearValueBinding(FLinearColor::Black);
+				OutputDesc.Flags &= (~ETextureCreateFlags::FastVRAM);
 
-	void AddGammaPass(FRDGBuilder& GraphBuilder, const FScreenPassViewInfo ViewInfo, FRDGTextureRef InOutTexture, float GammaPower)
-	{
-		RDG_EVENT_SCOPE_STAT(GraphBuilder, FHoldoutCompositeGamma, "HoldoutComposite.Gamma");
-		RDG_GPU_STAT_SCOPE(GraphBuilder, FHoldoutCompositeGamma);
-
-		FGlobalShaderMap* GlobalShaderMap = GetGlobalShaderMap(ViewInfo.FeatureLevel);
-		FScreenPassRenderTarget Output = FScreenPassRenderTarget(InOutTexture, ERenderTargetLoadAction::ENoAction);
-		FScreenPassTextureViewport Viewport = FScreenPassTextureViewport(InOutTexture);
-
-		FHoldoutCompositeGammaShader::FParameters* Parameters = GraphBuilder.AllocParameters<FHoldoutCompositeGammaShader::FParameters>();
-		Parameters->InputTexture = InOutTexture;
-		Parameters->InputSampler = TStaticSamplerState<>::GetRHI();
-		Parameters->GammaPower = GammaPower;
-		Parameters->RenderTargets[0] = Output.GetRenderTargetBinding();
-
-		TShaderMapRef<FHoldoutCompositeGammaShader> PixelShader(GlobalShaderMap);
-		AddDrawScreenPass(
-			GraphBuilder,
-			RDG_EVENT_NAME("HoldoutComposite.Gamma (%dx%d)", InOutTexture->Desc.Extent.X, InOutTexture->Desc.Extent.Y),
-			ViewInfo,
-			Viewport,
-			Viewport,
-			PixelShader,
-			Parameters
-		);
+				return OutputDesc;
+			}
+		}
 	}
 }
 
@@ -234,7 +125,7 @@ public:
 		const FRDGTextureDesc TextureDesc = FRDGTextureDesc::Create2D(RenderTargetTexture->Desc.Extent, PF_FloatRGBA, FClearValueBinding::Black, TexCreate_UAV | TexCreate_ShaderResource | TexCreate_RenderTargetable);
 		FRDGTextureRef Output = GraphBuilder.CreateTexture(TextureDesc, TEXT("HoldoutCompositeProcessedTexture"));
 
-		HoldoutComposite::AddDilatePass(GraphBuilder, RenderTargetTexture, Output, ViewFeatureLevel);
+		UE::HoldoutComposite::Private::AddDilatePass(GraphBuilder, RenderTargetTexture, Output, ViewFeatureLevel);
 
 		ParentExtension->CollectCustomRenderTarget(ViewId, GraphBuilder.ConvertToExternalTexture(Output));
 	}
@@ -249,28 +140,9 @@ private:
 FHoldoutCompositeSceneViewExtension::FHoldoutCompositeSceneViewExtension(const FAutoRegister& AutoReg, UWorld* InWorld)
 	: FWorldSceneViewExtension(AutoReg, InWorld)
 {
-	UHoldoutCompositePluginSettings* PluginSettings = GetMutableDefault<UHoldoutCompositePluginSettings>();
-	check(PluginSettings);
-
-#if WITH_EDITOR
-	DefautSettingsChangeHandle = PluginSettings->OnDefaultCompositeSettingsChanged.AddLambda([this]
-		{
-			OverrideCompositeSettings_GameThread(GetDefault<UHoldoutCompositePluginSettings>()->DefaultCompositeSettings, MIN_int32);
-		});
-#endif
-
-	CompositeSettings = CompositeSettings_RenderThread = PluginSettings->DefaultCompositeSettings;
 }
 
-FHoldoutCompositeSceneViewExtension::~FHoldoutCompositeSceneViewExtension()
-{
-#if WITH_EDITOR
-	if (UHoldoutCompositePluginSettings* PluginSettings = GetMutableDefault<UHoldoutCompositePluginSettings>())
-	{
-		PluginSettings->OnDefaultCompositeSettingsChanged.Remove(DefautSettingsChangeHandle);
-	}
-#endif
-}
+FHoldoutCompositeSceneViewExtension::~FHoldoutCompositeSceneViewExtension() = default;
 
 void FHoldoutCompositeSceneViewExtension::RegisterPrimitives_GameThread(const TArray<UPrimitiveComponent*>& InPrimitiveComponents)
 {
@@ -314,52 +186,43 @@ void FHoldoutCompositeSceneViewExtension::UnregisterPrimitives_GameThread(const
 	}
 }
 
-void FHoldoutCompositeSceneViewExtension::OverrideCompositeSettings_GameThread(const FHoldoutCompositeSettings& InSettings, int32 InPriority)
-{
-	check(IsInGameThread());
-
-	if (InPriority >= CurrentSettingsPriority)
-	{
-		UpdateSettings_GameThread(InSettings, InPriority);
-	}
-}
-
-void FHoldoutCompositeSceneViewExtension::ResetCompositeSettings_GameThread()
-{
-	check(IsInGameThread());
-
-	const UHoldoutCompositePluginSettings* PluginSettings = GetDefault<UHoldoutCompositePluginSettings>();
-	check(PluginSettings);
-
-	UpdateSettings_GameThread(PluginSettings->DefaultCompositeSettings, MIN_int32);
-}
-
-void FHoldoutCompositeSceneViewExtension::SetCompositePostProcessMaterials_GameThread(const TArray<UMaterialInterface*>& InMaterials)
+void FHoldoutCompositeSceneViewExtension::SetCompositeTextureOverride_GameThread(UTexture* InTextureOverride)
 {
-	check(IsInGameThread());
-
-	FScopeLock Lock(&MaterialsCS);
-	CompositeProcessMaterials.Reset();
-	CompositeProcessMaterials.Append(InMaterials);
+	CompositeTextureOverride = InTextureOverride;
 }
 
-void FHoldoutCompositeSceneViewExtension::ResetCompositePostProcessMaterials_GameThread()
+void FHoldoutCompositeSceneViewExtension::RegisterActorRenderWork_GameThread(UE::HoldoutComposite::Private::FRenderWork&& InWork, TOptional<int32> InViewUserFlagsOverride)
 {
-	check(IsInGameThread());
+	ViewUserFlagsOverride_GameThread = InViewUserFlagsOverride;
 
-	FScopeLock Lock(&MaterialsCS);
-	CompositeProcessMaterials.Reset();
+	ENQUEUE_RENDER_COMMAND(CopyHoldoutCompositeRenderWork)(
+		[InWork = MoveTemp(InWork), WeakThis = AsWeak()](FRHICommandListImmediate& RHICmdList) mutable
+		{
+			TSharedPtr<FHoldoutCompositeSceneViewExtension> SVE = StaticCastSharedPtr<FHoldoutCompositeSceneViewExtension>(WeakThis.Pin());
+			if (SVE.IsValid())
+			{
+				SVE->ActorWork_RenderThread = MoveTemp(InWork);
+			}
+		});
 }
 
-void FHoldoutCompositeSceneViewExtension::SetCompositeTextureOverride_GameThread(UTexture* InTextureOverride)
+void FHoldoutCompositeSceneViewExtension::RegisterActorRenderWork_GameThread()
 {
-	CompositeTextureOverride = InTextureOverride;
+	ENQUEUE_RENDER_COMMAND(CopyHoldoutCompositeRenderWork)(
+		[WeakThis = AsWeak()](FRHICommandListImmediate& RHICmdList)
+		{
+			TSharedPtr<FHoldoutCompositeSceneViewExtension> SVE = StaticCastSharedPtr<FHoldoutCompositeSceneViewExtension>(WeakThis.Pin());
+			if (SVE.IsValid())
+			{
+				SVE->ActorWork_RenderThread.Reset();
+			}
+		});
 }
 
 bool FHoldoutCompositeSceneViewExtension::IsActiveThisFrame_Internal(const FSceneViewExtensionContext& Context) const
 {
 	bool bIsActive = FWorldSceneViewExtension::IsActiveThisFrame_Internal(Context);
-	bIsActive &= !CompositePrimitives.IsEmpty() || CompositeTextureOverride.IsValid();
+	bIsActive &= !CompositePrimitives.IsEmpty();
 	bIsActive &= !IsHDREnabled();
 
 	return bIsActive;
@@ -370,7 +233,19 @@ bool FHoldoutCompositeSceneViewExtension::IsActiveThisFrame_Internal(const FScen
 
 int32 FHoldoutCompositeSceneViewExtension::GetPriority() const
 {
-	return GetSettings().CompositePassPriority;
+	return GetDefault<UHoldoutCompositePluginSettings>()->SceneViewExtensionPriority;
+}
+
+const UE::HoldoutComposite::Private::FRenderWork& FHoldoutCompositeSceneViewExtension::GetRenderWork() const
+{
+	if (ActorWork_RenderThread.IsSet())
+	{
+		return ActorWork_RenderThread.GetValue();
+	}
+	else
+	{
+		return UE::HoldoutComposite::Private::FRenderWork::GetDefault();
+	}
 }
 
 void FHoldoutCompositeSceneViewExtension::SetupViewFamily(FSceneViewFamily& InViewFamily)
@@ -380,12 +255,6 @@ void FHoldoutCompositeSceneViewExtension::SetupViewFamily(FSceneViewFamily& InVi
 
 void FHoldoutCompositeSceneViewExtension::SetupView(FSceneViewFamily& InViewFamily, FSceneView& InView)
 {
-	// Skip custom render pass since we'll use the composite texture override directly
-	if (CompositeTextureOverride.IsValid())
-	{
-		return;
-	}
-
 	// Disable custom render pass for view families without primitive alpha holdout support.
 	if (InViewFamily.EngineShowFlags.AllowPrimitiveAlphaHoldout == 0)
 	{
@@ -447,14 +316,20 @@ void FHoldoutCompositeSceneViewExtension::SetupView(FSceneViewFamily& InViewFami
 		InView
 	);
 
+	const UHoldoutCompositePluginSettings* PluginSettings = GetDefault<UHoldoutCompositePluginSettings>();
+	check(PluginSettings);
+
 	FSceneInterface::FCustomRenderPassRendererInput PassInput{};
 	PassInput.EngineShowFlags = InViewFamily.EngineShowFlags;
 	PassInput.EngineShowFlags.DisableFeaturesForUnlit();
 	PassInput.EngineShowFlags.SetTranslucency(true);
-	PassInput.EngineShowFlags.SetUnlitViewmode(GetSettings().bEnableUnlitViewmode);
+	PassInput.EngineShowFlags.SetUnlitViewmode(PluginSettings->bEnableUnlitViewmode);
 	PassInput.EngineShowFlags.SetAllowPrimitiveAlphaHoldout(false);
-	PassInput.bOverridesPostVolumeUserFlags = GetSettings().bOverridesViewUserFlags;
-	PassInput.PostVolumeUserFlags = GetSettings().CompositeViewUserFlags;
+	if (ViewUserFlagsOverride_GameThread.IsSet())
+	{
+		PassInput.bOverridesPostVolumeUserFlags = true;
+		PassInput.PostVolumeUserFlags = ViewUserFlagsOverride_GameThread.GetValue();
+	}
 	// Note: Incoming view location is invalid for scene captures
 	PassInput.ViewLocation = InView.ViewMatrices.GetViewOrigin();
 	PassInput.ViewRotationMatrix = InView.ViewMatrices.GetViewMatrix().RemoveTranslation();
@@ -470,258 +345,254 @@ void FHoldoutCompositeSceneViewExtension::SetupView(FSceneViewFamily& InViewFami
 	WorldPtr.Get()->Scene->AddCustomRenderPass(&InViewFamily, PassInput);
 }
 
-void FHoldoutCompositeSceneViewExtension::SubscribeToPostProcessingPass(EPostProcessingPass PassId, const FSceneView& InView, FAfterPassCallbackDelegateArray& InOutPassCallbacks, bool bIsPassEnabled)
+void FHoldoutCompositeSceneViewExtension::PreRenderViewFamily_RenderThread(FRDGBuilder& GraphBuilder, FSceneViewFamily& InViewFamily)
 {
-	constexpr bool bEvenIfPendingKill = false;
-	constexpr bool bThreadsafeTest = true;
-	bool bIsActive = static_cast<bool>(InView.Family->EngineShowFlags.AllowPrimitiveAlphaHoldout) &&
-		(CustomRenderTargetPerView_RenderThread.Contains(InView.GetViewKey()) || CompositeTextureOverride.IsValid(bEvenIfPendingKill, bThreadsafeTest));
+}
 
-	if (!bIsActive)
-	{
-		// Early-out to avoid needless work in the post processing callback(s).
-		return;
-	}
-	
-	if (PassId == EPostProcessingPass::SSRInput && GetSettings().bUpdateSSR)
-	{
-		InOutPassCallbacks.Add(FAfterPassCallbackDelegate::CreateRaw(this, &FHoldoutCompositeSceneViewExtension::PostProcessPassSSRInput_RenderThread));
-	}
-	else if (PassId == EPostProcessingPass::Tonemap && GetSettings().CompositePassLocation == EHoldoutCompositePassLocation::AfterTonemap)
-	{
-		InOutPassCallbacks.Add(FAfterPassCallbackDelegate::CreateRaw(this, &FHoldoutCompositeSceneViewExtension::PostProcessCompositePass_RenderThread));
-	}
-	else if (PassId == EPostProcessingPass::VisualizeDepthOfField && GetSettings().CompositePassLocation == EHoldoutCompositePassLocation::EndOfPostProcessing)
-	{
-		InOutPassCallbacks.Add(FAfterPassCallbackDelegate::CreateRaw(this, &FHoldoutCompositeSceneViewExtension::PostProcessCompositePass_RenderThread));
-	}
+void FHoldoutCompositeSceneViewExtension::PreRenderView_RenderThread(FRDGBuilder& GraphBuilder, FSceneView& InView)
+{
 }
 
-void FHoldoutCompositeSceneViewExtension::UpdateSettings_GameThread(const FHoldoutCompositeSettings& InSettings, int32 InPriority)
+void FHoldoutCompositeSceneViewExtension::PrePostProcessPass_RenderThread(FRDGBuilder& GraphBuilder, const FSceneView& InView, const FPostProcessingInputs& Inputs)
 {
-	CurrentSettingsPriority = InPriority;
-	CompositeSettings = InSettings;
-	
-	ENQUEUE_RENDER_COMMAND(CopyHoldoutCompositeSettings)(
-		[CompositeSettingsRT = InSettings, WeakThis = AsWeak()](FRHICommandListImmediate& RHICmdList)
-		{
-			const TSharedPtr<FHoldoutCompositeSceneViewExtension> SVE = StaticCastSharedPtr<FHoldoutCompositeSceneViewExtension>(WeakThis.Pin());
-			if (SVE.IsValid())
-			{
-				SVE->CompositeSettings_RenderThread = CompositeSettingsRT;
-			}
-		});
 }
 
-const FHoldoutCompositeSettings& FHoldoutCompositeSceneViewExtension::GetSettings() const
+void FHoldoutCompositeSceneViewExtension::SubscribeToPostProcessingPass(EPostProcessingPass PassId, const FSceneView& InView, FAfterPassCallbackDelegateArray& InOutPassCallbacks, bool bIsPassEnabled)
 {
-	if (IsInRenderingThread())
+	if (!IsActiveForView(InView))
 	{
-		return CompositeSettings_RenderThread;
+		return;
 	}
 	
-	checkf(IsInGameThread(), TEXT("Composite settings should only be access from the render thread and the game thread."));
-	return CompositeSettings;
-}
-
-FRDGTextureRef FHoldoutCompositeSceneViewExtension::GetCompositeTexture(FRDGBuilder& GraphBuilder, const FSceneView& InView) const
-{
-	if (FRDGTextureRef TextureOverrideRDG = HoldoutComposite::GetTextureOverride_RenderThread(GraphBuilder, CompositeTextureOverride))
+	if (PassId == EPostProcessingPass::BeforeDOF)
 	{
-		return TextureOverrideRDG;
+		// We use BeforeDOF as an input processing step for simplicity, since PrePostProcessPass_RenderThread has a different interface.
+		InOutPassCallbacks.Add(FAfterPassCallbackDelegate::CreateRaw(this, &FHoldoutCompositeSceneViewExtension::InputProcessing_RenderThread));
 	}
-
-	FRDGTextureRef CustomRenderPassTexture = GSystemTextures.GetBlackAlphaOneDummy(GraphBuilder);
-
-	const TRefCountPtr<IPooledRenderTarget>* CustomRenderPassRenderTargetPtr = CustomRenderTargetPerView_RenderThread.Find(InView.GetViewKey());
-	if (CustomRenderPassRenderTargetPtr != nullptr)
+	if (PassId == EPostProcessingPass::SSRInput && GetRenderWork().FramePasses.Contains(ISceneViewExtension::EPostProcessingPass::SSRInput))
 	{
-		CustomRenderPassTexture = GraphBuilder.RegisterExternalTexture(*CustomRenderPassRenderTargetPtr);
+		InOutPassCallbacks.Add(FAfterPassCallbackDelegate::CreateRaw(this, &FHoldoutCompositeSceneViewExtension::SSRInputWork_RenderThread));
 	}
-	
-	return CustomRenderPassTexture;
-}
-
-FHoldoutCompositeCommonParameters FHoldoutCompositeSceneViewExtension::BuildCommonCompositeParameters(FRDGBuilder& GraphBuilder, const FSceneView& InView, const FScreenPassTexture& SceneColor, const FScreenPassTexture& CompositeTexture, const FScreenPassRenderTarget& Output, bool bIsSceneColorUndistorted)
-{
-	const FLensDistortionLUT& LensDistortionLUT = LensDistortion::GetLUTUnsafe(InView);
-	const bool bLensDistortionInTSR = (LensDistortion::GetPassLocationUnsafe(InView) == LensDistortion::EPassLocation::TSR);
-
-	FHoldoutCompositeCommonParameters CommonParameters;
-	CommonParameters.Input = GetScreenPassTextureViewportParameters(FScreenPassTextureViewport(SceneColor));
-	CommonParameters.Composite = GetScreenPassTextureViewportParameters(FScreenPassTextureViewport(CompositeTexture));
-	CommonParameters.Output = GetScreenPassTextureViewportParameters(FScreenPassTextureViewport(Output));
-	CommonParameters.InputTexture = SceneColor.Texture;
-	CommonParameters.InputSampler = bIsSceneColorUndistorted ? TStaticSamplerState<SF_Bilinear, AM_Mirror, AM_Mirror>::GetRHI() : TStaticSamplerState<SF_Point>::GetRHI();
-	CommonParameters.CompositeTexture = CompositeTexture.Texture;
-	CommonParameters.CompositeSampler = TStaticSamplerState<SF_Bilinear>::GetRHI();
-	CommonParameters.UndistortingDisplacementTexture = GSystemTextures.GetBlackDummy(GraphBuilder);
-	CommonParameters.UndistortingDisplacementSampler = TStaticSamplerState<SF_Bilinear>::GetRHI();
-	CommonParameters.DistortingDisplacementTexture = GSystemTextures.GetBlackDummy(GraphBuilder);
-	CommonParameters.DistortingDisplacementSampler = TStaticSamplerState<SF_Bilinear>::GetRHI();
-	if (LensDistortionLUT.IsEnabled() && bLensDistortionInTSR)
+	else if (PassId == EPostProcessingPass::Tonemap && GetRenderWork().FramePasses.Contains(ISceneViewExtension::EPostProcessingPass::Tonemap))
 	{
-		CommonParameters.DistortingDisplacementTexture = LensDistortionLUT.DistortingDisplacementTexture;
-		CommonParameters.UndistortingDisplacementTexture = LensDistortionLUT.UndistortingDisplacementTexture;
+		InOutPassCallbacks.Add(FAfterPassCallbackDelegate::CreateRaw(this, &FHoldoutCompositeSceneViewExtension::TonemapWork_RenderThread));
+	}
+	else if (PassId == EPostProcessingPass::VisualizeDepthOfField && GetRenderWork().FramePasses.Contains(ISceneViewExtension::EPostProcessingPass::VisualizeDepthOfField))
+	{
+		InOutPassCallbacks.Add(FAfterPassCallbackDelegate::CreateRaw(this, &FHoldoutCompositeSceneViewExtension::EndOfPostProcessing_RenderThread));
 	}
-
-	return CommonParameters;
 }
 
-FScreenPassTexture FHoldoutCompositeSceneViewExtension::PostProcessPassSSRInput_RenderThread(FRDGBuilder& GraphBuilder, const FSceneView& InView, const FPostProcessMaterialInputs& Inputs)
+bool FHoldoutCompositeSceneViewExtension::IsActiveForView(const FSceneView& InView) const
 {
-	RDG_EVENT_SCOPE_STAT(GraphBuilder, FHoldoutCompositeSSRInput, "HoldoutComposite.SSRInput");
-	RDG_GPU_STAT_SCOPE(GraphBuilder, FHoldoutCompositeSSRInput);
-
-	FGlobalShaderMap* GlobalShaderMap = GetGlobalShaderMap(InView.GetFeatureLevel());
-
-	FScreenPassTexture SceneColor = FScreenPassTexture::CopyFromSlice(GraphBuilder, Inputs.GetInput(EPostProcessMaterialInput::SceneColor));
-	check(SceneColor.IsValid());
-	FScreenPassTexture CompositeTexture{ GetCompositeTexture(GraphBuilder, InView) };
-	FScreenPassRenderTarget Output = FScreenPassRenderTarget::CreateFromInput(GraphBuilder, SceneColor, InView.GetOverwriteLoadAction(), TEXT("HoldoutCompositeSSRInputRT"));
-
-	FHoldoutCompositeSSRInputShader::FParameters* PassParameters = GraphBuilder.AllocParameters<FHoldoutCompositeSSRInputShader::FParameters>();
-	constexpr bool bIsSceneColorUndistorted = true;
-	PassParameters->Common = BuildCommonCompositeParameters(GraphBuilder, InView, SceneColor, CompositeTexture, Output, bIsSceneColorUndistorted);
-	PassParameters->LastGlobalExposure = GetSettings().bFollowSceneExposure ? InView.GetLastEyeAdaptationExposure() : 1.0f;
-	PassParameters->RenderTargets[0] = Output.GetRenderTargetBinding();
-
-	TShaderMapRef<FHoldoutCompositeSSRInputShader> PixelShader(GlobalShaderMap);
-	FPixelShaderUtils::AddFullscreenPass(
-		GraphBuilder,
-		GlobalShaderMap,
-		RDG_EVENT_NAME("HoldoutComposite.SSRInput (%dx%d) PS",
-			Output.ViewRect.Width(), Output.ViewRect.Height()),
-		PixelShader,
-		PassParameters,
-		Output.ViewRect
-	);
+	constexpr bool bEvenIfPendingKill = false;
+	constexpr bool bThreadSafeTest = true;
 
-	return Output;
+	bool bIsActive = InView.Family->bIsMainViewFamily;
+	bIsActive &= !InView.bIsSceneCapture;
+	bIsActive &= static_cast<bool>(InView.Family->EngineShowFlags.AllowPrimitiveAlphaHoldout);
+	bIsActive &= CustomRenderTargetPerView_RenderThread.Contains(InView.GetViewKey()) || CompositeTextureOverride.IsValid(bEvenIfPendingKill, bThreadSafeTest);
+
+	return bIsActive;
 }
 
-FScreenPassTexture FHoldoutCompositeSceneViewExtension::PostProcessCompositePass_RenderThread(FRDGBuilder& GraphBuilder, const FSceneView& InView, const FPostProcessMaterialInputs& Inputs)
+FScreenPassTexture FHoldoutCompositeSceneViewExtension::InputProcessing_RenderThread(FRDGBuilder& GraphBuilder, const FSceneView& InView, const FPostProcessMaterialInputs& Inputs)
 {
-	RDG_EVENT_SCOPE_STAT(GraphBuilder, FHoldoutCompositeFinal, "HoldoutComposite.Final");
-	RDG_GPU_STAT_SCOPE(GraphBuilder, FHoldoutCompositeFinal);
-	using namespace HoldoutComposite;
-
-	FScreenPassTexture SceneColor = FScreenPassTexture::CopyFromSlice(GraphBuilder, Inputs.GetInput(EPostProcessMaterialInput::SceneColor));
-	check(SceneColor.IsValid());
+	using namespace UE::HoldoutComposite::Private;
 
-	FScreenPassTexture CompositeTexture{ GetCompositeTexture(GraphBuilder, InView) };
+	// Reset composite textures for view.
+	// NOTE: all post-processing SVE callbacks occur on *the same view*, so this is allowed.
+	CompositeRenderPass_RenderThread = CompositeTexture_RenderThread = FScreenPassTexture();
 
+	// Prepare texture inputs for following post-processing passes
 	{
-		FScopeLock Lock(&MaterialsCS);
-		for (const TWeakObjectPtr<UMaterialInterface>& MaterialPtr : CompositeProcessMaterials)
+		const TRefCountPtr<IPooledRenderTarget>* CompositeRenderPassPtr = CustomRenderTargetPerView_RenderThread.Find(InView.GetViewKey());
+		if (CompositeRenderPassPtr != nullptr)
 		{
-			TStrongObjectPtr<UMaterialInterface> Material = MaterialPtr.Pin();
-			if (Material.IsValid())
+			CompositeRenderPass_RenderThread = FScreenPassTexture{ GraphBuilder.RegisterExternalTexture(*CompositeRenderPassPtr) };
+		}
+
+		constexpr bool bEvenIfPendingKill = false;
+		constexpr bool bThreadSafeTest = true;
+		if (CompositeTextureOverride.IsValid(bEvenIfPendingKill, bThreadSafeTest))
+		{
+			if (FTextureResource* TextureResource = CompositeTextureOverride.GetEvenIfUnreachable()->GetResource())
 			{
-				FPostProcessMaterialInputs PostProcessMaterialInputs;
-				PostProcessMaterialInputs.SetInput(GraphBuilder, EPostProcessMaterialInput::SceneColor, CompositeTexture);
-				PostProcessMaterialInputs.SceneTextures = Inputs.SceneTextures; //Note: these are invalid as they come from the main render, but validation-required.
+				const FTextureRHIRef& TextureRHI = TextureResource->GetTextureRHI();
+				if (TextureRHI.IsValid())
+				{
+					FRDGTextureRef CompositeTextureRDG = GraphBuilder.FindExternalTexture(TextureRHI.GetReference());
 
-				// Create our own output override to maintain UAV access
-				FRDGTextureDesc OutputDesc = CompositeTexture.Texture->Desc;
-				OutputDesc.Reset();
-				OutputDesc.ClearValue = FClearValueBinding(FLinearColor::Black);
-				OutputDesc.Format = SceneColor.Texture->Desc.Format;
-				OutputDesc.Extent = CompositeTexture.ViewRect.Size();
-				OutputDesc.Flags |= ETextureCreateFlags::RenderTargetable;
-				PostProcessMaterialInputs.OverrideOutput = FScreenPassRenderTarget(GraphBuilder.CreateTexture(OutputDesc, TEXT("PostProcessMaterial")), CompositeTexture.ViewRect, InView.GetOverwriteLoadAction());
+					if (CompositeTextureRDG == nullptr)
+					{
+						CompositeTextureRDG = GraphBuilder.RegisterExternalTexture(CreateRenderTarget(TextureRHI.GetReference(), TEXT("CompositeTextureOverride")));
+					}
+					
+					FScreenPassTexture Input{ CompositeTextureRDG };
+					
+					const FScreenPassTextureSlice SceneColorSlice = Inputs.GetInput(EPostProcessMaterialInput::SceneColor);
+					const FScreenPassTextureViewport SceneColorViewport = FScreenPassTextureViewport(SceneColorSlice);
+
+					FRDGTextureDesc OutputDesc = GetPostProcessingDesc(SceneColorSlice.TextureSRV->Desc.Texture->Desc, Inputs.OutputFormat);
+					FScreenPassRenderTarget Output = FScreenPassRenderTarget(GraphBuilder.CreateTexture(OutputDesc, TEXT("CompositeTextureCopy")), SceneColorViewport.Rect, InView.GetOverwriteLoadAction());
+	
+					// TODO: We could avoid the copy if the texture override format is usable as scene color.
+					{
+						const FScreenPassTextureViewport InputViewport(Input);
+						const FScreenPassTextureViewport OutputViewport(Output);
 
-				CompositeTexture = AddPostProcessMaterialPass(GraphBuilder, InView, PostProcessMaterialInputs, Material.Get());
+						TShaderMapRef<FCopyRectPS> PixelShader(GetGlobalShaderMap(InView.GetFeatureLevel()));
+
+						FCopyRectPS::FParameters* Parameters = GraphBuilder.AllocParameters<FCopyRectPS::FParameters>();
+						Parameters->InputTexture = Input.Texture;
+						Parameters->InputSampler = TStaticSamplerState<SF_Bilinear>::GetRHI();
+						Parameters->RenderTargets[0] = Output.GetRenderTargetBinding();
+
+						AddDrawScreenPass(GraphBuilder, RDG_EVENT_NAME("DrawTexture"), InView, OutputViewport, InputViewport, PixelShader, Parameters);
+					}
+
+					CompositeTexture_RenderThread = Output;
+				}
 			}
 		}
 	}
 
-	if (GetSettings().bApplyFXAA)
+	// Texture processing
+	for (TPair<EHoldoutCompositePassInput, TArray<const IHoldoutCompositePassProxy*>> Processing : GetRenderWork().InputProcessing)
 	{
-		if (GetSettings().bLinearComposite)
+		const EHoldoutCompositePassInput ActiveTextureInput = Processing.Key;
+		const TArray<const IHoldoutCompositePassProxy*>& Passes = Processing.Value;
+
+		if (!ensureMsgf(ActiveTextureInput != EHoldoutCompositePassInput::Default, TEXT("Processing is expected to run on composite textures.")))
+		{
+			continue;
+		}
+
+		TSortedMap<EHoldoutCompositePassInput, FScreenPassTextureSlice> TextureInputs;
+		if (CompositeRenderPass_RenderThread.IsValid())
 		{
-			// We gamma-encode the result so that FXAA can operate on non-linear colors
-			HoldoutComposite::AddGammaPass(GraphBuilder, InView, CompositeTexture.Texture, 1.0f / HoldoutComposite::DefaultGamma);
+			TextureInputs.Emplace(EHoldoutCompositePassInput::CompositeRenderPass, FScreenPassTextureSlice::CreateFromScreenPassTexture(GraphBuilder, CompositeRenderPass_RenderThread));
 		}
+		if (CompositeTexture_RenderThread.IsValid())
+		{
+			TextureInputs.Emplace(EHoldoutCompositePassInput::CompositeTexture, FScreenPassTextureSlice::CreateFromScreenPassTexture(GraphBuilder, CompositeTexture_RenderThread));
+		}
+
+		FScreenPassTexture Output;
+
+		for (const IHoldoutCompositePassProxy* Pass : Passes)
+		{
+			if (FScreenPassTextureSlice* SliceInput = TextureInputs.Find(ActiveTextureInput))
+			{
+				FPostProcessMaterialInputs PassInputs = Inputs;
+				PassInputs.SetInput(EPostProcessMaterialInput::SceneColor, *SliceInput);
+				PassInputs.OverrideOutput = FScreenPassRenderTarget{};
 
-		static TConsoleVariableData<int32>* CVarFXAAQuality = IConsoleManager::Get().FindTConsoleVariableDataInt(TEXT("r.FXAA.Quality"));
+				// Register pass and update output
+				Output = Pass->Add(GraphBuilder, InView, PassInputs, ISceneViewExtension::EPostProcessingPass::BeforeDOF);
 
-		FFXAAInputs PassInputs{};
-		PassInputs.SceneColor = CompositeTexture;
-		PassInputs.Quality = static_cast<EFXAAQuality>(FMath::Clamp(CVarFXAAQuality->GetValueOnRenderThread(), 0, 5));
-		CompositeTexture = AddFXAAPass(GraphBuilder, InView, PassInputs);
+				// Update active input for the next pass
+				TextureInputs[ActiveTextureInput] = FScreenPassTextureSlice::CreateFromScreenPassTexture(GraphBuilder, Output);
+			}
+		}
+
+		// Update composite textures for later callbacks in post-processing
+		if (Output.IsValid())
+		{
+			switch (ActiveTextureInput)
+			{
+			case EHoldoutCompositePassInput::CompositeRenderPass:
+				CompositeRenderPass_RenderThread = Output;
+				break;
+			case EHoldoutCompositePassInput::CompositeTexture:
+				CompositeTexture_RenderThread = Output;
+				break;
+
+			default:
+				checkNoEntry();
+				break;
+			}
+		}
 	}
 
-	FGlobalShaderMap* GlobalShaderMap = GetGlobalShaderMap(InView.GetFeatureLevel());
-	const FSceneViewFamily* Family = InView.Family;
+	return Inputs.ReturnUntouchedSceneColorForPostProcessing(GraphBuilder);
+}
 
-	FUint32Vector Encodings = FUint32Vector::ZeroValue; // ESceneColorSourceEncoding::Linear
-	if ((Family->EngineShowFlags.Tonemapper == 0) || (Family->EngineShowFlags.PostProcessing == 0))
+FScreenPassTexture FHoldoutCompositeSceneViewExtension::PostProcessWork_RenderThread(FRDGBuilder& GraphBuilder, const FSceneView& InView, const FPostProcessMaterialInputs& Inputs, ISceneViewExtension::EPostProcessingPass InLocation)
+{
+	using namespace UE::HoldoutComposite::Private;
+
+	TSortedMap<EHoldoutCompositePassInput, FScreenPassTextureSlice> TextureInputs;
+	TextureInputs.Emplace(EHoldoutCompositePassInput::Default, Inputs.GetInput(EPostProcessMaterialInput::SceneColor));
+	if (CompositeRenderPass_RenderThread.IsValid())
 	{
-		Encodings.X = static_cast<uint32>(ESourceEncoding::Gamma);
-		Encodings.Y = static_cast<uint32>(ESourceEncoding::Gamma);
+		TextureInputs.Emplace(EHoldoutCompositePassInput::CompositeRenderPass, FScreenPassTextureSlice::CreateFromScreenPassTexture(GraphBuilder, CompositeRenderPass_RenderThread));
 	}
-	else if (Family->SceneCaptureSource == SCS_FinalColorLDR)
+	if (CompositeTexture_RenderThread.IsValid())
 	{
-		Encodings.X = static_cast<uint32>(ESourceEncoding::sRGB);
-		Encodings.Y = static_cast<uint32>(ESourceEncoding::sRGB);
+		TextureInputs.Emplace(EHoldoutCompositePassInput::CompositeTexture, FScreenPassTextureSlice::CreateFromScreenPassTexture(GraphBuilder, CompositeTexture_RenderThread));
 	}
 
-	// Encodings.Z controls the encoding of the composite texture
-	// A default value of 0 means that the composite is in linear color space
-	if (!GetSettings().bLinearComposite)
+	FScreenPassTexture Output;
+	const TArray<const IHoldoutCompositePassProxy*>& Passes = GetRenderWork().FramePasses[InLocation];
+	for (int32 Index = 0; Index < Passes.Num(); ++Index)
 	{
-		// The composite texture is in the sRGB color space and must be encoded in the linear color space
-		Encodings.Z = static_cast<uint32>(ESourceEncoding::sRGB);
+		const IHoldoutCompositePassProxy* Pass = Passes[Index];
+		const bool bIsLastPass = (Index == Passes.Num() - 1);
+
+		// Update pass inputs
+		FPostProcessMaterialInputs PassInputs = Inputs;
+		int32 ConnectedInputs = 0;
+		for (int32 InputIndex = 0; InputIndex < Pass->GetNumInputs() && InputIndex < TextureInputs.Num(); ++InputIndex)
+		{
+			if (FScreenPassTextureSlice* SliceInput = TextureInputs.Find(Pass->GetPassInput(InputIndex)))
+			{
+				PassInputs.SetInput(static_cast<EPostProcessMaterialInput>(InputIndex), *SliceInput);
+				ConnectedInputs++;
+			}
+		}
+
+		if (ConnectedInputs == 0)
+		{
+			continue;
+		}
+
+		// Only apply the output override on the last pass
+		if (!bIsLastPass)
+		{
+			PassInputs.OverrideOutput = FScreenPassRenderTarget{};
+		}
+
+		// Register pass and update output
+		Output = Pass->Add(GraphBuilder, InView, PassInputs, InLocation);
+
+		// Update input 0 for the next pass
+		TextureInputs[EHoldoutCompositePassInput::Default] = FScreenPassTextureSlice::CreateFromScreenPassTexture(GraphBuilder, Output);
 	}
-	else if(GetSettings().bApplyFXAA)
+
+	if (Output.IsValid())
 	{
-		// FXAA received gamma-encoded colors, so we must linearize before compositing.
-		Encodings.Z = static_cast<uint32>(ESourceEncoding::Gamma);
+		return MoveTemp(Output);
 	}
-
-
-	FScreenPassRenderTarget Output = Inputs.OverrideOutput;
-	if (!Output.IsValid())
+	else
 	{
-		Output = FScreenPassRenderTarget::CreateFromInput(GraphBuilder, SceneColor, InView.GetOverwriteLoadAction(), TEXT("HoldoutCompositePassOutput"));
+		return Inputs.ReturnUntouchedSceneColorForPostProcessing(GraphBuilder);
 	}
+}
 
-	// Compositing pass
-	{
-		FHoldoutCompositeFinalShader::FPermutationDomain PermutationVector;
-		PermutationVector.Set<FHoldoutCompositeFinalShader::FUseGlobalExposure>(GetSettings().bFollowSceneExposure);
-		constexpr bool bEvenIfPendingKill = false;
-		constexpr bool bThreadsafeTest = true;
-		const bool bActiveTextureOverride = CompositeTextureOverride.IsValid(bEvenIfPendingKill, bThreadsafeTest);
+FScreenPassTexture FHoldoutCompositeSceneViewExtension::SSRInputWork_RenderThread(FRDGBuilder& GraphBuilder, const FSceneView& InView, const FPostProcessMaterialInputs& Inputs)
+{
+	return PostProcessWork_RenderThread(GraphBuilder, InView, Inputs, ISceneViewExtension::EPostProcessingPass::SSRInput);
+}
 
-		FHoldoutCompositeFinalShader::FParameters* PassParameters = GraphBuilder.AllocParameters<FHoldoutCompositeFinalShader::FParameters>();
-		PassParameters->View = InView.ViewUniformBuffer;
-		PassParameters->Common = BuildCommonCompositeParameters(GraphBuilder, InView, SceneColor, CompositeTexture, Output);
-		if (GetSettings().bFollowSceneExposure)
-		{
-			FRDGBufferRef EyeAdaptationBuffer = GraphBuilder.RegisterExternalBuffer(InView.GetEyeAdaptationBuffer(), ERDGBufferFlags::MultiFrame);
-			PassParameters->EyeAdaptationBuffer = GraphBuilder.CreateSRV(EyeAdaptationBuffer);
-		}
-		PassParameters->Encodings = Encodings;
-		PassParameters->DisplayGamma = FVector3f(Family->RenderTarget->GetDisplayGamma(), 1.0f / Family->RenderTarget->GetDisplayGamma(), HoldoutComposite::DefaultGamma);
-		PassParameters->bApplyTSRLensDistortion = GetSettings().bApplyTSRLensDistortion;
-		PassParameters->bIsCustomRenderPassTexture = !bActiveTextureOverride;
-		PassParameters->RenderTargets[0] = Output.GetRenderTargetBinding();
-
-		TShaderMapRef<FHoldoutCompositeFinalShader> PixelShader(GlobalShaderMap, PermutationVector);
-		FPixelShaderUtils::AddFullscreenPass(
-			GraphBuilder,
-			GlobalShaderMap,
-			RDG_EVENT_NAME("HoldoutComposite.Final (%dx%d) PS",
-				Output.ViewRect.Width(), Output.ViewRect.Height()),
-			PixelShader,
-			PassParameters,
-			Output.ViewRect
-		);
-	}
+FScreenPassTexture FHoldoutCompositeSceneViewExtension::TonemapWork_RenderThread(FRDGBuilder& GraphBuilder, const FSceneView& InView, const FPostProcessMaterialInputs& Inputs)
+{
+	return PostProcessWork_RenderThread(GraphBuilder, InView, Inputs, ISceneViewExtension::EPostProcessingPass::Tonemap);
+}
 
-	return MoveTemp(Output);
+FScreenPassTexture FHoldoutCompositeSceneViewExtension::EndOfPostProcessing_RenderThread(FRDGBuilder& GraphBuilder, const FSceneView& InView, const FPostProcessMaterialInputs& Inputs)
+{
+	return PostProcessWork_RenderThread(GraphBuilder, InView, Inputs, ISceneViewExtension::EPostProcessingPass::VisualizeDepthOfField);
 }
 
 void FHoldoutCompositeSceneViewExtension::PostRenderViewFamily_RenderThread(FRDGBuilder& GraphBuilder, FSceneViewFamily& InViewFamily)
diff --git a/Engine/Plugins/Experimental/Compositing/HoldoutComposite/Source/HoldoutComposite/Private/HoldoutCompositeSceneViewExtension.h b/Engine/Plugins/Experimental/Compositing/HoldoutComposite/Source/HoldoutComposite/Private/HoldoutCompositeSceneViewExtension.h
--- a/Engine/Plugins/Experimental/Compositing/HoldoutComposite/Source/HoldoutComposite/Private/HoldoutCompositeSceneViewExtension.h
+++ b/Engine/Plugins/Experimental/Compositing/HoldoutComposite/Source/HoldoutComposite/Private/HoldoutCompositeSceneViewExtension.h
@@ -3,17 +3,18 @@
 #pragma once
 
 #include "SceneViewExtension.h"
+
+#include "Containers/ContainersFwd.h"
 #include "RenderGraphFwd.h"
 #include "RendererInterface.h"
-#include "Containers/ContainersFwd.h"
+#include "ScreenPass.h"
 
 #include "HoldoutCompositeSettings.h"
+#include "Passes/HoldoutCompositePassBase.h"
 
-enum class EHoldoutCompositePassLocation;
 class FViewInfo;
 struct FScreenPassRenderTarget;
 struct FScreenPassTexture;
-struct FHoldoutCompositeSettings;
 class UMaterialInterface;
 class UTexture;
 
@@ -29,21 +30,15 @@ public:
 	/* Unregister primitives for compositing. */
 	void UnregisterPrimitives_GameThread(const TArray<UPrimitiveComponent*>& InPrimitiveComponents);
 
-	/* Override the composite settings. */
-	void OverrideCompositeSettings_GameThread(const FHoldoutCompositeSettings& InSettings, int32 InPriority);
-
-	/* Reset composite settings back to plugin settings default. */
-	void ResetCompositeSettings_GameThread();
-
-	/* Set additional post-process materials to be applied on the separate composited render. */
-	void SetCompositePostProcessMaterials_GameThread(const TArray<UMaterialInterface*>& InMaterials);
-
-	/* Remove all post-process materials to be applied on the separate composited render. */
-	void ResetCompositePostProcessMaterials_GameThread();
-
 	/* Set a texture override to be used instead of the separate composited render. */
 	void SetCompositeTextureOverride_GameThread(UTexture* InTextureOverride);
 
+	/* Register render work to be enqueued on the render thread, and an optional view user flags override. */
+	void RegisterActorRenderWork_GameThread(UE::HoldoutComposite::Private::FRenderWork&& InWork, TOptional<int32> InViewUserFlagsOverride = {});
+
+	/* Reset custom actor render work. */
+	void RegisterActorRenderWork_GameThread();
+
 	/* Called by the custom render pass to store its view render target for this frame. */
 	template<typename T>
 	void CollectCustomRenderTarget(uint32 InViewId, T InRenderTarget)
@@ -55,6 +50,9 @@ public:
 	virtual int32 GetPriority() const override;
 	virtual void SetupViewFamily(FSceneViewFamily& InViewFamily) override;
 	virtual void SetupView(FSceneViewFamily& InViewFamily, FSceneView& InView) override;
+	virtual void PreRenderViewFamily_RenderThread(FRDGBuilder& GraphBuilder, FSceneViewFamily& InViewFamily) override;
+	virtual void PreRenderView_RenderThread(FRDGBuilder& GraphBuilder, FSceneView& InView) override;
+	virtual void PrePostProcessPass_RenderThread(FRDGBuilder& GraphBuilder, const FSceneView& InView, const FPostProcessingInputs& Inputs) override;
 	virtual void SubscribeToPostProcessingPass(EPostProcessingPass PassId, const FSceneView& InView, FAfterPassCallbackDelegateArray& InOutPassCallbacks, bool bIsPassEnabled) override;
 	virtual void PostRenderViewFamily_RenderThread(FRDGBuilder& GraphBuilder, FSceneViewFamily& InViewFamily) override;
 	virtual void PostRenderView_RenderThread(FRDGBuilder& GraphBuilder, FSceneView& InView) override;
@@ -65,48 +63,39 @@ protected:
 	virtual bool IsActiveThisFrame_Internal(const FSceneViewExtensionContext& Context) const override;
 	//~ End ISceneViewExtension Interface
 
-	/* Update composite settings. */
-	void UpdateSettings_GameThread(const FHoldoutCompositeSettings& InSettings, int32 InPriority);
-
-	/* Thread-safe getter for composite settings. */
-	const FHoldoutCompositeSettings& GetSettings() const;
-
 private:
+	/* Getter for the frame render work, either from the actor or default. */
+	const UE::HoldoutComposite::Private::FRenderWork& GetRenderWork() const;
 
-	/* Returns the custom render pass render target, or the composite texture override if present.*/
-	FRDGTextureRef GetCompositeTexture(FRDGBuilder& GraphBuilder, const FSceneView& InView) const;
-
-	/* Build the common composite pass shader parameters. */
-	class FHoldoutCompositeCommonParameters BuildCommonCompositeParameters(FRDGBuilder& GraphBuilder, const FSceneView& InView, const FScreenPassTexture& SceneColor, const FScreenPassTexture& CompositeTexture, const FScreenPassRenderTarget& Output, bool bIsSceneColorUndistorted = false);
+	/* Active check to avoid work in post-processing. */
+	bool IsActiveForView(const FSceneView& InView) const;
 
-	FScreenPassTexture PostProcessPassSSRInput_RenderThread(FRDGBuilder& GraphBuilder, const FSceneView& InView, const FPostProcessMaterialInputs& Inputs);
-	FScreenPassTexture PostProcessCompositePass_RenderThread(FRDGBuilder& GraphBuilder, const FSceneView& InView, const FPostProcessMaterialInputs& Inputs);
+	/* Callbacks for processing passes. */
+	FScreenPassTexture InputProcessing_RenderThread(FRDGBuilder& GraphBuilder, const FSceneView& InView, const FPostProcessMaterialInputs& Inputs);
+	FScreenPassTexture PostProcessWork_RenderThread(FRDGBuilder& GraphBuilder, const FSceneView& InView, const FPostProcessMaterialInputs& Inputs, ISceneViewExtension::EPostProcessingPass InLocation);
+	FScreenPassTexture SSRInputWork_RenderThread(FRDGBuilder& GraphBuilder, const FSceneView& InView, const FPostProcessMaterialInputs& Inputs);
+	FScreenPassTexture TonemapWork_RenderThread(FRDGBuilder& GraphBuilder, const FSceneView& InView, const FPostProcessMaterialInputs& Inputs);
+	FScreenPassTexture EndOfPostProcessing_RenderThread(FRDGBuilder& GraphBuilder, const FSceneView& InView, const FPostProcessMaterialInputs& Inputs);
 
 	/* Collection of primitives to render as a custom render pass and composite after post-processing. .*/
 	TSet<TWeakObjectPtr<UPrimitiveComponent>> CompositePrimitives;
 
+	/* Composite texture override use directly instead of the separate composited render. */
+	TWeakObjectPtr<UTexture> CompositeTextureOverride;
+
 	/* Custom render pass render targets for each active view .*/
 	TMap<uint32, TRefCountPtr<IPooledRenderTarget>> CustomRenderTargetPerView_RenderThread;
 
-	/* Composite settings and its render-thread copy */
-	FHoldoutCompositeSettings CompositeSettings;
-	FHoldoutCompositeSettings CompositeSettings_RenderThread;
-
-	/* Composite settings priority */
-	int32 CurrentSettingsPriority = MIN_int32;
-
-	/* Composite texture override use directly instead of the separate composited render. */
-	TWeakObjectPtr<UTexture> CompositeTextureOverride;
+	/* Custom actor render work. */
+	TOptional<UE::HoldoutComposite::Private::FRenderWork> ActorWork_RenderThread;
 
-	/** Post-process material(s) to be applied on the separate composited render.*/
-	TArray<TWeakObjectPtr<UMaterialInterface>> CompositeProcessMaterials;
+	/* Composite render pass screen texture, only valid in post-processing callbacks. */
+	FScreenPassTexture CompositeRenderPass_RenderThread;
 
-	/** Critical section for composite post-process materials. */
-	FCriticalSection MaterialsCS;
+	/* Composite screen texture, only valid in post-processing callbacks. */
+	FScreenPassTexture CompositeTexture_RenderThread;
 
-#if WITH_EDITOR
-	/** Delegate handle for default plugin settings change. */
-	FDelegateHandle DefautSettingsChangeHandle;
-#endif
+	/* Optional view user flags override. */
+	TOptional<int32> ViewUserFlagsOverride_GameThread;
 };
 
diff --git a/Engine/Plugins/Experimental/Compositing/HoldoutComposite/Source/HoldoutComposite/Private/HoldoutCompositeSettings.cpp b/Engine/Plugins/Experimental/Compositing/HoldoutComposite/Source/HoldoutComposite/Private/HoldoutCompositeSettings.cpp
--- a/Engine/Plugins/Experimental/Compositing/HoldoutComposite/Source/HoldoutComposite/Private/HoldoutCompositeSettings.cpp
+++ b/Engine/Plugins/Experimental/Compositing/HoldoutComposite/Source/HoldoutComposite/Private/HoldoutCompositeSettings.cpp
@@ -15,19 +15,6 @@ static FAutoConsoleVariableRef CVarRefreshPrimitivesOnTick(
 	TEXT("Read-only and to be set in a config file (requires restart)."),
 	ECVF_ReadOnly);
 
-FHoldoutCompositeSettings::FHoldoutCompositeSettings()
-	: bFollowSceneExposure(false)
-	, bUpdateSSR(true)
-	, bEnableUnlitViewmode(true)
-	, bApplyTSRLensDistortion(true)
-	, bApplyFXAA(false)
-	, CompositePassLocation(EHoldoutCompositePassLocation::AfterTonemap)
-	, CompositePassPriority(100 + 1 /* OPENCOLORIO_SCENE_VIEW_EXTENSION_PRIORITY + 1 */)
-	, bOverridesViewUserFlags(false)
-	, CompositeViewUserFlags(0)
-	, bLinearComposite(true)
-{ }
-
 UHoldoutCompositePluginSettings::UHoldoutCompositePluginSettings()
 	: DisabledPrimitiveClasses(
 		{
@@ -38,9 +25,11 @@ UHoldoutCompositePluginSettings::UHoldoutCompositePluginSettings()
 			FSoftClassPath("/Script/Engine.LineBatchComponent")
 		}
 	)
-	, DefaultCompositeSettings()
+	, bEnableUnlitViewmode(true)
+	, SceneViewExtensionPriority(100 + 1 /* OPENCOLORIO_SCENE_VIEW_EXTENSION_PRIORITY + 1 */)
 { }
 
+
 FName UHoldoutCompositePluginSettings::GetCategoryName() const
 {
 	return TEXT("Plugins");
@@ -60,11 +49,7 @@ void UHoldoutCompositePluginSettings::PostEditChangeProperty(FPropertyChangedEve
 {
 	Super::PostEditChangeProperty(PropertyChangedEvent);
 
-	if (PropertyChangedEvent.GetMemberPropertyName() == GET_MEMBER_NAME_CHECKED(ThisClass, DefaultCompositeSettings))
-	{
-		OnDefaultCompositeSettingsChanged.Broadcast();
-	}
-	else if (PropertyChangedEvent.GetMemberPropertyName() == GET_MEMBER_NAME_CHECKED(ThisClass, DisabledPrimitiveClasses))
+	if (PropertyChangedEvent.GetMemberPropertyName() == GET_MEMBER_NAME_CHECKED(ThisClass, DisabledPrimitiveClasses))
 	{
 		CacheDisabledClasses();
 	}
diff --git a/Engine/Plugins/Experimental/Compositing/HoldoutComposite/Source/HoldoutComposite/Private/HoldoutCompositeSubsystem.cpp b/Engine/Plugins/Experimental/Compositing/HoldoutComposite/Source/HoldoutComposite/Private/HoldoutCompositeSubsystem.cpp
--- a/Engine/Plugins/Experimental/Compositing/HoldoutComposite/Source/HoldoutComposite/Private/HoldoutCompositeSubsystem.cpp
+++ b/Engine/Plugins/Experimental/Compositing/HoldoutComposite/Source/HoldoutComposite/Private/HoldoutCompositeSubsystem.cpp
@@ -1,8 +1,11 @@
 // Copyright Epic Games, Inc. All Rights Reserved.
 
 #include "HoldoutCompositeSubsystem.h"
+
+#include "HoldoutComposite.h"
 #include "HoldoutCompositeModule.h"
 #include "HoldoutCompositeSceneViewExtension.h"
+#include "Passes/HoldoutCompositePassBase.h"
 
 #include "Engine/RendererSettings.h"
 
@@ -66,67 +69,94 @@ void UHoldoutCompositeSubsystem::Deinitialize()
 	Super::Deinitialize();
 }
 
-void UHoldoutCompositeSubsystem::RegisterPrimitive(UPrimitiveComponent* InPrimitiveComponent)
+void UHoldoutCompositeSubsystem::Tick(float DeltaTime)
 {
-	TArray<UPrimitiveComponent*> PrimitiveComponents = { InPrimitiveComponent };
+	Super::Tick(DeltaTime);
 
-	RegisterPrimitives(PrimitiveComponents);
+	if (!HoldoutCompositeViewExtension.IsValid())
+	{
+		return;
+	}
+
+	if (AHoldoutComposite* CompositeActor = CompositeActorPtr.Get())
+	{
+		TOptional<int32> UserFlagsOverride = CompositeActor->bOverridesViewUserFlags ? CompositeActor->ViewUserFlags : TOptional<int32>{};
+
+		// Every frame, we record  render work which the SVE moves for render-thread only access.
+		UE::HoldoutComposite::Private::FRenderWork RenderWork;
+		for (FHoldoutCompositeProcessingPasses& CompositeProcessing : CompositeActor->CompositeProcessing)
+		{
+			RenderWork.AddCompositeProcessing(CompositeProcessing.InputToProcess, CompositeProcessing.Passes);
+		}
+		RenderWork.AddPostProcessing(ISceneViewExtension::EPostProcessingPass::SSRInput, CompositeActor->Passes_SSRInput);
+		RenderWork.AddPostProcessing(ISceneViewExtension::EPostProcessingPass::Tonemap, CompositeActor->Passes_AfterToneMap);
+		RenderWork.AddPostProcessing(ISceneViewExtension::EPostProcessingPass::VisualizeDepthOfField, CompositeActor->Passes_AfterPostProcessing);
+		HoldoutCompositeViewExtension->RegisterActorRenderWork_GameThread(MoveTemp(RenderWork), UserFlagsOverride);
+	}
 }
 
-void UHoldoutCompositeSubsystem::RegisterPrimitives(const TArray<UPrimitiveComponent*>& InPrimitiveComponents)
+void UHoldoutCompositeSubsystem::RegisterActor(AHoldoutComposite* InActor)
 {
-	const bool bValidSettings = ValidateProjectSettings();
+	CompositeActorPtr = InActor;
 
-	if (bValidSettings && HoldoutCompositeViewExtension.IsValid())
+#if WITH_EDITOR
+	CompositeActorCount++;
+
+	if (CompositeActorCount > 1)
 	{
-		HoldoutCompositeViewExtension->RegisterPrimitives_GameThread(InPrimitiveComponents);
+		if (GEngine)
+		{
+			GEngine->AddOnScreenDebugMessage(-1, 5.0f, FColor::Red, TEXT("No more than one holdout composite actor is allowed."));
+		}
+
+		UE_LOG(LogHoldoutComposite, Error, TEXT("No more than one holdout composite actor is allowed."));
 	}
+#endif // WITH_EDITOR
 }
 
-void UHoldoutCompositeSubsystem::UnregisterPrimitive(UPrimitiveComponent* InPrimitiveComponent)
+void UHoldoutCompositeSubsystem::ResetActor()
 {
-	TArray<UPrimitiveComponent*> PrimitiveComponents = { InPrimitiveComponent };
+	CompositeActorPtr.Reset();
 
-	UnregisterPrimitives(PrimitiveComponents);
-}
+#if WITH_EDITOR
+	CompositeActorCount--;
+#endif
 
-void UHoldoutCompositeSubsystem::UnregisterPrimitives(const TArray<UPrimitiveComponent*>& InPrimitiveComponents)
-{
 	if (HoldoutCompositeViewExtension.IsValid())
 	{
-		HoldoutCompositeViewExtension->UnregisterPrimitives_GameThread(InPrimitiveComponents);
+		HoldoutCompositeViewExtension->RegisterActorRenderWork_GameThread();
 	}
 }
 
-void UHoldoutCompositeSubsystem::SetCompositePostProcessMaterials(const TArray<UMaterialInterface*>& InMaterials)
+void UHoldoutCompositeSubsystem::RegisterPrimitive(UPrimitiveComponent* InPrimitiveComponent)
 {
-	if (HoldoutCompositeViewExtension.IsValid())
-	{
-		HoldoutCompositeViewExtension->SetCompositePostProcessMaterials_GameThread(InMaterials);
-	}
+	TArray<UPrimitiveComponent*> PrimitiveComponents = { InPrimitiveComponent };
+
+	RegisterPrimitives(PrimitiveComponents);
 }
 
-void UHoldoutCompositeSubsystem::ResetCompositePostProcessMaterials()
+void UHoldoutCompositeSubsystem::RegisterPrimitives(const TArray<UPrimitiveComponent*>& InPrimitiveComponents)
 {
-	if (HoldoutCompositeViewExtension.IsValid())
+	const bool bValidSettings = ValidateProjectSettings();
+
+	if (bValidSettings && HoldoutCompositeViewExtension.IsValid())
 	{
-		HoldoutCompositeViewExtension->ResetCompositePostProcessMaterials_GameThread();
+		HoldoutCompositeViewExtension->RegisterPrimitives_GameThread(InPrimitiveComponents);
 	}
 }
 
-void UHoldoutCompositeSubsystem::OverrideCompositeSettings(const FHoldoutCompositeSettings& InSettings, int32 InPriority)
+void UHoldoutCompositeSubsystem::UnregisterPrimitive(UPrimitiveComponent* InPrimitiveComponent)
 {
-	if (HoldoutCompositeViewExtension.IsValid())
-	{
-		HoldoutCompositeViewExtension->OverrideCompositeSettings_GameThread(InSettings, InPriority);
-	}
+	TArray<UPrimitiveComponent*> PrimitiveComponents = { InPrimitiveComponent };
+
+	UnregisterPrimitives(PrimitiveComponents);
 }
 
-void UHoldoutCompositeSubsystem::ResetCompositeSettings()
+void UHoldoutCompositeSubsystem::UnregisterPrimitives(const TArray<UPrimitiveComponent*>& InPrimitiveComponents)
 {
 	if (HoldoutCompositeViewExtension.IsValid())
 	{
-		HoldoutCompositeViewExtension->ResetCompositeSettings_GameThread();
+		HoldoutCompositeViewExtension->UnregisterPrimitives_GameThread(InPrimitiveComponents);
 	}
 }
 
@@ -138,11 +168,6 @@ void UHoldoutCompositeSubsystem::SetCompositeTextureOverride(UTexture* InTexture
 	}
 }
 
-bool UHoldoutCompositeSubsystem::EnableHoldoutSettingsNotification() const
-{
-	return true;
-}
-
 bool UHoldoutCompositeSubsystem::IsProjectSettingsValid()
 {
 	const URendererSettings* RendererSettings = GetDefault<URendererSettings>();
@@ -159,7 +184,7 @@ bool UHoldoutCompositeSubsystem::ValidateProjectSettings()
 {
 	const bool bValidSettings = IsProjectSettingsValid();
 
-	if (!bValidSettings && EnableHoldoutSettingsNotification())
+	if (!bValidSettings )
 	{
 #if WITH_EDITOR
 		URendererSettings* RendererSettings = GetMutableDefault<URendererSettings>();
diff --git a/Engine/Plugins/Experimental/Compositing/HoldoutComposite/Source/HoldoutComposite/Private/Passes/HoldoutCompositePassBase.cpp b/Engine/Plugins/Experimental/Compositing/HoldoutComposite/Source/HoldoutComposite/Private/Passes/HoldoutCompositePassBase.cpp
new file mode 100644
--- /dev/null
+++ b/Engine/Plugins/Experimental/Compositing/HoldoutComposite/Source/HoldoutComposite/Private/Passes/HoldoutCompositePassBase.cpp
@@ -0,0 +1,127 @@
+// Copyright Epic Games, Inc. All Rights Reserved.
+
+#include "Passes/HoldoutCompositePassBase.h"
+
+#include "HoldoutComposite.h"
+#include "Passes/HoldoutCompositePassMerge.h"
+
+namespace UE
+{
+	namespace HoldoutComposite
+	{
+		namespace Private
+		{
+			FRenderWork::FRenderWork()
+				: InputProcessing{}
+				, FramePasses{}
+				, FrameAllocator(new FSceneRenderingBulkObjectAllocator{})
+			{ }
+
+			const FRenderWork& FRenderWork::GetDefault()
+			{
+				static FRenderWork DefaultRenderWork = []() -> FRenderWork
+					{
+						FRenderWork RenderWork;
+
+						constexpr EHoldoutCompositePassInput Input0 = EHoldoutCompositePassInput::Default;
+						constexpr EHoldoutCompositePassInput Input1 = EHoldoutCompositePassInput::CompositeRenderPass;
+
+						TArray<const IHoldoutCompositePassProxy*>& SSRInputWork = RenderWork.FramePasses.Add(ISceneViewExtension::EPostProcessingPass::SSRInput);
+						TArray<const IHoldoutCompositePassProxy*>& AfterTonemapWork = RenderWork.FramePasses.Add(ISceneViewExtension::EPostProcessingPass::Tonemap);
+
+						SSRInputWork.Add(RenderWork.FrameAllocator->Create<FHoldoutCompositePassMergeProxy>(Input0, Input1));
+						AfterTonemapWork.Add(RenderWork.FrameAllocator->Create<FHoldoutCompositePassMergeProxy>(Input0, Input1));
+
+						return RenderWork;
+					}();
+
+				return DefaultRenderWork;
+			}
+
+			void FRenderWork::AddCompositeProcessing(EHoldoutCompositePassInput InputType, TArrayView<TObjectPtr<UHoldoutCompositePassBase>> InPasses)
+			{
+				TPair<EHoldoutCompositePassInput, TArray<const IHoldoutCompositePassProxy*>> Pair = { InputType, TArray<const IHoldoutCompositePassProxy*>{} };
+
+				RegisterPasses(InPasses, Pair.Value);
+
+				InputProcessing.Add(MoveTemp(Pair));
+			}
+
+			void FRenderWork::AddPostProcessing(ISceneViewExtension::EPostProcessingPass InLocation, TArrayView<TObjectPtr<UHoldoutCompositePassBase>> InPasses)
+			{
+				if (InPasses.IsEmpty())
+				{
+					return;
+				}
+
+				TArray<const IHoldoutCompositePassProxy*>& Proxies = FramePasses.Add(InLocation);
+				RegisterPasses(InPasses, Proxies);
+			}
+
+			void FRenderWork::RegisterPasses(TArrayView<TObjectPtr<UHoldoutCompositePassBase>> InPasses, TArray<const IHoldoutCompositePassProxy*>& OutProxies)
+			{
+				for (const TObjectPtr<UHoldoutCompositePassBase>& Pass : InPasses)
+				{
+					if (IsValid(Pass))
+					{
+						if (!Pass->bIsEnabled)
+						{
+							continue;
+						}
+
+						IHoldoutCompositePassProxy* OutProxy;
+						if (Pass->GetProxy(*FrameAllocator, OutProxy))
+						{
+							OutProxies.Add(OutProxy);
+						}
+					}
+				}
+			}
+		}
+	}
+}
+
+UHoldoutCompositePassBase::UHoldoutCompositePassBase(const FObjectInitializer& ObjectInitializer)
+	: Super(ObjectInitializer)
+	, OwnerPrivate(GetTypedOuter<AHoldoutComposite>())
+{
+
+}
+
+UHoldoutCompositePassBase::~UHoldoutCompositePassBase() = default;
+
+#if WITH_EDITOR
+bool UHoldoutCompositePassBase::CanEditChange(const FProperty* InProperty) const
+{
+	if (InProperty)
+	{
+		// We don't support changing this property as it is fully controlled by Niagara in both batached and solo mode
+		if (InProperty->GetFName() == GET_MEMBER_NAME_CHECKED(UHoldoutCompositePassBase, Input0))
+		{
+			return bCanEditInput;
+		}
+	}
+
+	return Super::CanEditChange(InProperty);
+}
+
+void UHoldoutCompositePassBase::SetCanEditInput(bool bInEdit)
+{
+	bCanEditInput = bInEdit;
+}
+
+TArray<FString> UHoldoutCompositePassBase::GetRestrictedInputs() const
+{
+	if (IsValid(OwnerPrivate))
+	{
+		const UTexture* Texture = OwnerPrivate->GetCompositeTexture();
+		if (Texture != nullptr)
+		{
+			return {};
+		}
+	}
+
+	return { TEXT("CompositeTexture") };
+}
+#endif
+
diff --git a/Engine/Plugins/Experimental/Compositing/HoldoutComposite/Source/HoldoutComposite/Private/Passes/HoldoutCompositePassFXAA.cpp b/Engine/Plugins/Experimental/Compositing/HoldoutComposite/Source/HoldoutComposite/Private/Passes/HoldoutCompositePassFXAA.cpp
new file mode 100644
--- /dev/null
+++ b/Engine/Plugins/Experimental/Compositing/HoldoutComposite/Source/HoldoutComposite/Private/Passes/HoldoutCompositePassFXAA.cpp
@@ -0,0 +1,120 @@
+// Copyright Epic Games, Inc. All Rights Reserved.
+
+#include "Passes/HoldoutCompositePassFXAA.h"
+
+#include "RenderGraphBuilder.h"
+#include "RenderGraphUtils.h"
+#include "SceneView.h"
+#include "ScreenPass.h"
+#include "PostProcess/PostProcessAA.h"
+#include "PostProcess/PostProcessMaterialInputs.h"
+
+class FHoldoutCompositeGammaShader : public FGlobalShader
+{
+	DECLARE_GLOBAL_SHADER(FHoldoutCompositeGammaShader);
+	SHADER_USE_PARAMETER_STRUCT(FHoldoutCompositeGammaShader, FGlobalShader);
+
+	BEGIN_SHADER_PARAMETER_STRUCT(FParameters, )
+		SHADER_PARAMETER_RDG_TEXTURE(Texture2D, InputTexture)
+		SHADER_PARAMETER_SAMPLER(SamplerState, InputSampler)
+		SHADER_PARAMETER(FIntPoint, Dimensions)
+		SHADER_PARAMETER(float, GammaPower)
+		RENDER_TARGET_BINDING_SLOTS()
+	END_SHADER_PARAMETER_STRUCT()
+};
+IMPLEMENT_GLOBAL_SHADER(FHoldoutCompositeGammaShader, "/Plugin/HoldoutComposite/Private/HoldoutCompositeGamma.usf", "MainPS", SF_Pixel);
+
+DECLARE_GPU_STAT_NAMED(FHoldoutCompositeGamma, TEXT("HoldoutComposite.Gamma"));
+
+namespace UE
+{
+	namespace HoldoutComposite
+	{
+		namespace Private
+		{
+			void AddGammaPass(FRDGBuilder& GraphBuilder, const FScreenPassViewInfo ViewInfo, FRDGTextureRef InOutTexture, float GammaPower)
+			{
+				RDG_EVENT_SCOPE_STAT(GraphBuilder, FHoldoutCompositeGamma, "HoldoutComposite.Gamma");
+				RDG_GPU_STAT_SCOPE(GraphBuilder, FHoldoutCompositeGamma);
+
+				FGlobalShaderMap* GlobalShaderMap = GetGlobalShaderMap(ViewInfo.FeatureLevel);
+				FScreenPassRenderTarget Output = FScreenPassRenderTarget(InOutTexture, ERenderTargetLoadAction::ENoAction);
+				FScreenPassTextureViewport Viewport = FScreenPassTextureViewport(InOutTexture);
+
+				FHoldoutCompositeGammaShader::FParameters* Parameters = GraphBuilder.AllocParameters<FHoldoutCompositeGammaShader::FParameters>();
+				Parameters->InputTexture = InOutTexture;
+				Parameters->InputSampler = TStaticSamplerState<>::GetRHI();
+				Parameters->GammaPower = GammaPower;
+				Parameters->RenderTargets[0] = Output.GetRenderTargetBinding();
+
+				TShaderMapRef<FHoldoutCompositeGammaShader> PixelShader(GlobalShaderMap);
+				AddDrawScreenPass(
+					GraphBuilder,
+					RDG_EVENT_NAME("HoldoutComposite.Gamma (%dx%d)", InOutTexture->Desc.Extent.X, InOutTexture->Desc.Extent.Y),
+					ViewInfo,
+					Viewport,
+					Viewport,
+					PixelShader,
+					Parameters
+				);
+			}
+		}
+	}
+}
+
+
+class FHoldoutCompositeFXAAPassProxy : public IHoldoutCompositePassProxy
+{
+public:
+	using IHoldoutCompositePassProxy::IHoldoutCompositePassProxy;
+
+	bool bLinearSourceColors;
+	int32 Quality;
+
+	FScreenPassTexture Add(FRDGBuilder& GraphBuilder, const FSceneView& InView, const FPostProcessMaterialInputs& Inputs, ISceneViewExtension::EPostProcessingPass InLocation) const override
+	{
+		const FScreenPassTexture Input = FScreenPassTexture::CopyFromSlice(GraphBuilder, Inputs.GetInput(EPostProcessMaterialInput::SceneColor));
+
+		static constexpr float DefaultGamma = 2.2f;
+
+		if (bLinearSourceColors)
+		{
+			// We gamma-encode the result so that FXAA can operate on non-linear colors
+			UE::HoldoutComposite::Private::AddGammaPass(GraphBuilder, InView, Input.Texture, 1.0f / DefaultGamma);
+		}
+
+		FFXAAInputs PassInputs{};
+		PassInputs.SceneColor = Input;
+		PassInputs.Quality = static_cast<EFXAAQuality>(Quality);
+		PassInputs.OverrideOutput = Inputs.OverrideOutput;
+
+		FScreenPassTexture Output = AddFXAAPass(GraphBuilder, InView, PassInputs);
+
+		if (bLinearSourceColors)
+		{
+			UE::HoldoutComposite::Private::AddGammaPass(GraphBuilder, InView, Output.Texture, DefaultGamma);
+		}
+
+		return MoveTemp(Output);
+	}
+};
+
+UHoldoutCompositeFXAAPass::UHoldoutCompositeFXAAPass(const FObjectInitializer& ObjectInitializer)
+	: Super(ObjectInitializer)
+	, bLinearSourceColors(true)
+	, Quality(5)
+{
+}
+
+UHoldoutCompositeFXAAPass::~UHoldoutCompositeFXAAPass() = default;
+
+bool UHoldoutCompositeFXAAPass::GetProxy(FSceneRenderingBulkObjectAllocator& InFrameAllocator, IHoldoutCompositePassProxy*& OutProxy) const
+{
+	FHoldoutCompositeFXAAPassProxy* Proxy = InFrameAllocator.Create<FHoldoutCompositeFXAAPassProxy>(Input0);
+	Proxy->Quality = Quality;
+	Proxy->bLinearSourceColors = bLinearSourceColors;
+
+	OutProxy = Proxy;
+	return true;
+}
+
diff --git a/Engine/Plugins/Experimental/Compositing/HoldoutComposite/Source/HoldoutComposite/Private/Passes/HoldoutCompositePassMaterial.cpp b/Engine/Plugins/Experimental/Compositing/HoldoutComposite/Source/HoldoutComposite/Private/Passes/HoldoutCompositePassMaterial.cpp
new file mode 100644
--- /dev/null
+++ b/Engine/Plugins/Experimental/Compositing/HoldoutComposite/Source/HoldoutComposite/Private/Passes/HoldoutCompositePassMaterial.cpp
@@ -0,0 +1,59 @@
+// Copyright Epic Games, Inc. All Rights Reserved.
+
+#include "Passes/HoldoutCompositePassMaterial.h"
+
+#include "Materials/MaterialInterface.h"
+#include "RenderGraphBuilder.h"
+#include "RenderGraphUtils.h"
+#include "SceneView.h"
+#include "ScreenPass.h"
+#include "PostProcess/PostProcessMaterialInputs.h"
+
+class FHoldoutCompositePassMaterialProxy : public IHoldoutCompositePassProxy
+{
+public:
+	using IHoldoutCompositePassProxy::IHoldoutCompositePassProxy;
+
+	TWeakObjectPtr<UMaterialInterface> Material;
+
+	FScreenPassTexture Add(FRDGBuilder& GraphBuilder, const FSceneView& InView, const FPostProcessMaterialInputs& Inputs, ISceneViewExtension::EPostProcessingPass InLocation) const override
+	{
+		constexpr bool bEvenIfPendingKill = false;
+		constexpr bool bThreadsafeTest = true;
+		if (Material.IsValid(bEvenIfPendingKill, bThreadsafeTest))
+		{
+			return AddPostProcessMaterialPass(GraphBuilder, InView, Inputs, Material.GetEvenIfUnreachable());
+		}
+
+		if (Inputs.OverrideOutput.IsValid())
+		{
+			return Inputs.OverrideOutput;
+		}
+
+		return FScreenPassTexture::CopyFromSlice(GraphBuilder, Inputs.GetInput(EPostProcessMaterialInput::SceneColor));
+	}
+};
+
+
+UHoldoutCompositePassMaterial::UHoldoutCompositePassMaterial(const FObjectInitializer& ObjectInitializer)
+	: Super(ObjectInitializer)
+{
+
+}
+
+UHoldoutCompositePassMaterial::~UHoldoutCompositePassMaterial() = default;
+
+bool UHoldoutCompositePassMaterial::GetProxy(FSceneRenderingBulkObjectAllocator& InFrameAllocator, IHoldoutCompositePassProxy*& OutProxy) const
+{
+	if (IsValid(PostProcessMaterial))
+	{
+		FHoldoutCompositePassMaterialProxy* Proxy = InFrameAllocator.Create<FHoldoutCompositePassMaterialProxy>(Input0);
+		Proxy->Material = PostProcessMaterial;
+		
+		OutProxy = Proxy;
+		return true;
+	}
+
+	return false;
+}
+
diff --git a/Engine/Plugins/Experimental/Compositing/HoldoutComposite/Source/HoldoutComposite/Private/Passes/HoldoutCompositePassMerge.cpp b/Engine/Plugins/Experimental/Compositing/HoldoutComposite/Source/HoldoutComposite/Private/Passes/HoldoutCompositePassMerge.cpp
new file mode 100644
--- /dev/null
+++ b/Engine/Plugins/Experimental/Compositing/HoldoutComposite/Source/HoldoutComposite/Private/Passes/HoldoutCompositePassMerge.cpp
@@ -0,0 +1,235 @@
+// Copyright Epic Games, Inc. All Rights Reserved.
+
+#include "Passes/HoldoutCompositePassMerge.h"
+
+#include "HoldoutCompositeSettings.h"
+
+#include "PixelShaderUtils.h"
+#include "PostProcess/LensDistortion.h"
+#include "PostProcess/PostProcessMaterialInputs.h"
+#include "RenderGraphBuilder.h"
+#include "RenderGraphUtils.h"
+#include "SceneView.h"
+#include "SystemTextures.h"
+
+DECLARE_GPU_STAT_NAMED(FHoldoutCompositeMerge, TEXT("HoldoutComposite.Merge"));
+
+class FHoldoutCompositeMergeShader : public FGlobalShader
+{
+	DECLARE_GLOBAL_SHADER(FHoldoutCompositeMergeShader);
+	SHADER_USE_PARAMETER_STRUCT(FHoldoutCompositeMergeShader, FGlobalShader);
+
+	class FUseDistortion : SHADER_PERMUTATION_BOOL("USE_DISTORTION");
+	using FPermutationDomain = TShaderPermutationDomain<FUseDistortion>;
+
+	BEGIN_SHADER_PARAMETER_STRUCT(FInputPassTextureParameters, )
+		SHADER_PARAMETER_RDG_TEXTURE(Texture2D, Texture)
+		SHADER_PARAMETER_SAMPLER(SamplerState, Sampler)
+		SHADER_PARAMETER(uint32, bAlphaIsOpacity)
+		SHADER_PARAMETER(uint32, SourceEncoding)
+		SHADER_PARAMETER(uint32, DistortionUV)
+	END_SHADER_PARAMETER_STRUCT()
+
+	BEGIN_SHADER_PARAMETER_STRUCT(FParameters, )
+		SHADER_PARAMETER_STRUCT_REF(FViewUniformShaderParameters, View)
+		SHADER_PARAMETER_STRUCT(FScreenPassTextureViewportParameters, Input0)
+		SHADER_PARAMETER_STRUCT(FScreenPassTextureViewportParameters, Input1)
+		SHADER_PARAMETER_STRUCT(FScreenPassTextureViewportParameters, Output)
+		SHADER_PARAMETER_STRUCT(FInputPassTextureParameters, Tex0)
+		SHADER_PARAMETER_STRUCT(FInputPassTextureParameters, Tex1)
+		SHADER_PARAMETER_RDG_TEXTURE(Texture2D, DistortingDisplacementTexture)
+		SHADER_PARAMETER_SAMPLER(SamplerState, DistortingDisplacementSampler)
+		SHADER_PARAMETER_RDG_TEXTURE(Texture2D, UndistortingDisplacementTexture)
+		SHADER_PARAMETER_SAMPLER(SamplerState, UndistortingDisplacementSampler)
+		SHADER_PARAMETER(FVector2f, DisplayGamma)
+		SHADER_PARAMETER(uint32, MergeOp)
+		SHADER_PARAMETER(uint32, OutputEncoding)
+		SHADER_PARAMETER(uint32, bFlags)
+		RENDER_TARGET_BINDING_SLOTS()
+	END_SHADER_PARAMETER_STRUCT()
+};
+IMPLEMENT_GLOBAL_SHADER(FHoldoutCompositeMergeShader, "/Plugin/HoldoutComposite/Private/HoldoutCompositeMerge.usf", "MainPS", SF_Pixel);
+
+namespace UE
+{
+	namespace HoldoutComposite
+	{
+		namespace Private
+		{
+			enum class ESourceEncoding : uint32
+			{
+				Linear = 0,
+				Gamma = 1,
+				sRGB = 2,
+			};
+
+			enum class EDistortionUV
+			{
+				None = 0,
+				Distorted = 1,
+				Undistorted = 2
+			};
+
+			FHoldoutCompositeMergeShader::FInputPassTextureParameters GetPassTextureParameters(
+				const FScreenPassTexture& InTexture,
+				const EHoldoutCompositePassInput InputType,
+				const FSceneViewFamily* Family,
+				const ISceneViewExtension::EPostProcessingPass InLocation,
+				const bool bUseLensDistortion)
+			{
+				FHoldoutCompositeMergeShader::FInputPassTextureParameters Parameters;
+
+				Parameters.Texture = InTexture.Texture;
+				Parameters.Sampler = TStaticSamplerState<SF_Bilinear>::GetRHI();
+				Parameters.bAlphaIsOpacity = (InputType == EHoldoutCompositePassInput::CompositeTexture) ? 1 : 0;
+				Parameters.SourceEncoding = static_cast<uint32>(ESourceEncoding::Linear);
+				
+				// After-tonemap scene color may have encoding manually applied, as opposed to _SRGB textures
+				if(InputType == EHoldoutCompositePassInput::Default && InLocation >= ISceneViewExtension::EPostProcessingPass::Tonemap)
+				{
+					if ((Family->EngineShowFlags.Tonemapper == 0) || (Family->EngineShowFlags.PostProcessing == 0))
+					{
+						Parameters.SourceEncoding = static_cast<uint32>(ESourceEncoding::Gamma);
+					}
+					else if (Family->SceneCaptureSource == SCS_FinalColorLDR)
+					{
+						Parameters.SourceEncoding = static_cast<uint32>(ESourceEncoding::sRGB);
+					}
+				}
+
+				EDistortionUV DistortionUV = EDistortionUV::None;
+
+				if (bUseLensDistortion)
+				{
+					const bool bOutputIsDistorted = !(InLocation == ISceneViewExtension::EPostProcessingPass::SSRInput ||
+						InLocation == ISceneViewExtension::EPostProcessingPass::TranslucencyAfterDOF);
+
+					if (bOutputIsDistorted)
+					{
+						// If output is distorted, we must sample undistorted textures with undistorted UVs
+						if (InputType == EHoldoutCompositePassInput::CompositeRenderPass)
+						{
+							DistortionUV = EDistortionUV::Undistorted;
+						}
+					}
+					else
+					{
+						// If output is undistorted, we must sample distorted textures with distorted UVs
+						if (InputType == EHoldoutCompositePassInput::Default || InputType == EHoldoutCompositePassInput::CompositeTexture)
+						{
+							DistortionUV = EDistortionUV::Distorted;
+						}
+					}
+				}
+
+				Parameters.DistortionUV = static_cast<uint32>(DistortionUV);
+
+				return Parameters;
+			}
+
+			FHoldoutCompositePassMergeProxy::FHoldoutCompositePassMergeProxy(EHoldoutCompositePassInput Input0, EHoldoutCompositePassInput Input1)
+				: IHoldoutCompositePassProxy(Input0)
+			{
+				InputTypes.Add(Input1);
+			}
+
+			FScreenPassTexture FHoldoutCompositePassMergeProxy::Add(FRDGBuilder& GraphBuilder, const FSceneView& InView, const FPostProcessMaterialInputs& Inputs, ISceneViewExtension::EPostProcessingPass InLocation) const
+			{
+				RDG_EVENT_SCOPE_STAT(GraphBuilder, FHoldoutCompositeMerge, "HoldoutComposite.Merge");
+				RDG_GPU_STAT_SCOPE(GraphBuilder, FHoldoutCompositeMerge);
+
+				FScreenPassTexture Input0 = FScreenPassTexture::CopyFromSlice(GraphBuilder, Inputs.GetInput(static_cast<EPostProcessMaterialInput>(0)));
+				FScreenPassTexture Input1 = FScreenPassTexture::CopyFromSlice(GraphBuilder, Inputs.GetInput(static_cast<EPostProcessMaterialInput>(1)));
+
+				FGlobalShaderMap* GlobalShaderMap = GetGlobalShaderMap(InView.GetFeatureLevel());
+				const FSceneViewFamily* Family = InView.Family;
+
+				ESourceEncoding OutputEncoding = ESourceEncoding::Linear;
+
+				// When running after tonemap, we assume that input0 (scene color) has an encoding.
+				if (InLocation >= ISceneViewExtension::EPostProcessingPass::Tonemap)
+				{
+					if ((Family->EngineShowFlags.Tonemapper == 0) || (Family->EngineShowFlags.PostProcessing == 0))
+					{
+						OutputEncoding = ESourceEncoding::Gamma;
+					}
+					else if (Family->SceneCaptureSource == SCS_FinalColorLDR)
+					{
+						OutputEncoding = ESourceEncoding::sRGB;
+					}
+				}
+
+				FScreenPassRenderTarget Output = Inputs.OverrideOutput;
+				if (!Output.IsValid())
+				{
+					Output = FScreenPassRenderTarget::CreateFromInput(GraphBuilder, Input0, InView.GetOverwriteLoadAction(), TEXT("HoldoutCompositeMergeOutput"));
+				}
+
+				const FLensDistortionLUT& LensDistortionLUT = LensDistortion::GetLUTUnsafe(InView);
+				const bool bUseLensDistortion = LensDistortionLUT.IsEnabled() && (LensDistortion::GetPassLocationUnsafe(InView) == LensDistortion::EPassLocation::TSR);
+
+				FHoldoutCompositeMergeShader::FParameters* PassParameters = GraphBuilder.AllocParameters<FHoldoutCompositeMergeShader::FParameters>();
+				PassParameters->View = InView.ViewUniformBuffer;
+				PassParameters->Input0 = GetScreenPassTextureViewportParameters(FScreenPassTextureViewport(Input0));
+				PassParameters->Input1 = GetScreenPassTextureViewportParameters(FScreenPassTextureViewport(Input1));
+				PassParameters->Output = GetScreenPassTextureViewportParameters(FScreenPassTextureViewport(Output));
+				PassParameters->Tex0 = GetPassTextureParameters(Input0, InputTypes[0], Family, InLocation, bUseLensDistortion);
+				PassParameters->Tex1 = GetPassTextureParameters(Input1, InputTypes[1], Family, InLocation, bUseLensDistortion);
+				if (bUseLensDistortion)
+				{
+					PassParameters->DistortingDisplacementTexture = LensDistortionLUT.DistortingDisplacementTexture;
+					PassParameters->UndistortingDisplacementTexture = LensDistortionLUT.UndistortingDisplacementTexture;
+				}
+				else
+				{
+					PassParameters->UndistortingDisplacementTexture = GSystemTextures.GetBlackDummy(GraphBuilder);
+					PassParameters->DistortingDisplacementTexture = GSystemTextures.GetBlackDummy(GraphBuilder);
+				}
+				PassParameters->UndistortingDisplacementSampler = TStaticSamplerState<SF_Bilinear>::GetRHI();
+				PassParameters->DistortingDisplacementSampler = TStaticSamplerState<SF_Bilinear>::GetRHI();
+				PassParameters->DisplayGamma = FVector2f(Family->RenderTarget->GetDisplayGamma(), 1.0f / Family->RenderTarget->GetDisplayGamma());
+				PassParameters->MergeOp = static_cast<uint32>(MergeOp);
+				PassParameters->OutputEncoding = static_cast<uint32>(OutputEncoding);
+				PassParameters->bFlags = static_cast<uint32>(bUnpremulInput1 << 1) | static_cast<uint32>(bOutputSceneColor);
+				PassParameters->RenderTargets[0] = Output.GetRenderTargetBinding();
+
+				FHoldoutCompositeMergeShader::FPermutationDomain PermutationVector;
+				PermutationVector.Set<FHoldoutCompositeMergeShader::FUseDistortion>(bUseLensDistortion);
+				TShaderMapRef<FHoldoutCompositeMergeShader> PixelShader(GlobalShaderMap, PermutationVector);
+
+				FPixelShaderUtils::AddFullscreenPass(
+					GraphBuilder,
+					GlobalShaderMap,
+					RDG_EVENT_NAME("HoldoutComposite.Merge (%dx%d) PS",
+						Output.ViewRect.Width(), Output.ViewRect.Height()),
+					PixelShader,
+					PassParameters,
+					Output.ViewRect
+				);
+
+				return MoveTemp(Output);
+			}
+		}
+	}
+}
+
+UHoldoutCompositePassMerge::UHoldoutCompositePassMerge(const FObjectInitializer& ObjectInitializer)
+	: Super(ObjectInitializer)
+{
+}
+
+UHoldoutCompositePassMerge::~UHoldoutCompositePassMerge() = default;
+
+bool UHoldoutCompositePassMerge::GetProxy(FSceneRenderingBulkObjectAllocator& InFrameAllocator, IHoldoutCompositePassProxy*& OutProxy) const
+{
+	using namespace UE::HoldoutComposite::Private;
+
+	FHoldoutCompositePassMergeProxy* Proxy = InFrameAllocator.Create<FHoldoutCompositePassMergeProxy>(Input0, Input1);
+	Proxy->MergeOp = MergeOp;
+	Proxy->bUnpremulInput1 = bUnpremulInput1;
+	Proxy->bOutputSceneColor = true;
+
+	OutProxy = Proxy;
+	return true;
+}
+
diff --git a/Engine/Plugins/Experimental/Compositing/HoldoutComposite/Source/HoldoutComposite/Public/HoldoutComposite.h b/Engine/Plugins/Experimental/Compositing/HoldoutComposite/Source/HoldoutComposite/Public/HoldoutComposite.h
--- a/Engine/Plugins/Experimental/Compositing/HoldoutComposite/Source/HoldoutComposite/Public/HoldoutComposite.h
+++ b/Engine/Plugins/Experimental/Compositing/HoldoutComposite/Source/HoldoutComposite/Public/HoldoutComposite.h
@@ -6,6 +6,8 @@
 #include "GameFramework/Actor.h"
 #include "HoldoutCompositeSettings.h"
 
+#include "Passes/HoldoutCompositePassBase.h"
+
 #include "HoldoutComposite.generated.h"
 
 class ACameraActor;
@@ -15,12 +17,29 @@ class UMaterialInterface;
 class UMaterialInstanceDynamic;
 class UTexture;
 
+/* Struct used to pre-process composite texture inputs with a series of passes. */
+USTRUCT(BlueprintType)
+struct HOLDOUTCOMPOSITE_API FHoldoutCompositeProcessingPasses
+{
+	GENERATED_BODY()
+
+	/* Composite texture input type. */
+	UPROPERTY(EditAnywhere, BlueprintReadWrite, Category = "", meta = (ValidEnumValues = "CompositeRenderPass, CompositeTexture"))
+	EHoldoutCompositePassInput InputToProcess = EHoldoutCompositePassInput::CompositeRenderPass;
+
+	/* Passes applied onto the input texture. */
+	UPROPERTY(EditAnywhere, BlueprintReadWrite, Instanced, Category = "", meta = (DisallowedClasses = "/Script/HoldoutComposite.HoldoutCompositePassMerge"))
+	TArray<TObjectPtr<UHoldoutCompositePassBase>> Passes;
+};
+
+/* Composite actor used to control properties of the built-in rendering & (holdout) compositing functionality. */
 UCLASS()
 class HOLDOUTCOMPOSITE_API AHoldoutComposite : public AActor
 {
 	GENERATED_BODY()
 
 public:
+
 	AHoldoutComposite(const FObjectInitializer& ObjectInitializer);
 	~AHoldoutComposite();
 
@@ -34,7 +53,7 @@ public:
 
 #if WITH_EDITOR
 	virtual void PreEditChange(FProperty* PropertyThatWillChange) override;
-	virtual void PostEditChangeProperty(FPropertyChangedEvent& PropertyChangedEvent) override;
+	virtual void PostEditChangeChainProperty(struct FPropertyChangedChainEvent& PropertyChangedEvent) override;
 	virtual void SetIsTemporarilyHiddenInEditor(bool bIsHidden) override;
 	virtual bool ShouldTickIfViewportsOnly() const override;
 #endif // WITH_EDITOR
@@ -55,37 +74,22 @@ public:
 	UFUNCTION(BlueprintSetter)
 	void SetCompositeTexture(UTexture* InCompositeTexture);
 
-	/* Set the override settings. */
-	UFUNCTION(BlueprintSetter)
-	void SetEnableOverrideSettings(bool InEnable);
-
-	/* Get the override settings. */
-	UFUNCTION(BlueprintGetter)
-	bool GetEnableOverrideSettings() const;
-
 private:
-	/* Whether or not the holdout composite is active. */
+	/* Whether or not the holdout composite behavior is active. */
 	UPROPERTY(EditAnywhere, BlueprintGetter = IsEnabled, BlueprintSetter = SetEnabled, Category = "HoldoutComposite", meta = (AllowPrivateAccess = true, DisplayPriority="0"))
 	bool bIsEnabled = true;
 
 public:
-	/* Actors whose primitive (mesh) components will be marked as primitive alpha holdout, rendered separately and composited with scene color after post-processing. */
+	/* Actors whose primitive mesh components will be marked as primitive alpha holdout, rendered separately and composited with scene color during post-processing. */
 	UPROPERTY(EditAnywhere, BlueprintReadWrite, Category = "HoldoutComposite")
 	TArray<TObjectPtr<AActor>> CompositeMeshActors;
 
-	/** Post-process material(s) to be applied on the separate composited render.*/
-	UPROPERTY(EditAnywhere, BlueprintReadWrite, Category = "HoldoutComposite")
-	TArray<TObjectPtr<UMaterialInterface>> CompositePostProcessMaterials;
-
 private:
-	/* Texture to be used directly for alpha holdout composite and/or to be projected onto composite meshes. */
+	/* Texture to be used directly as an input for the alpha holdout composite. */
 	UPROPERTY(EditAnywhere, BlueprintGetter = GetCompositeTexture, BlueprintSetter = SetCompositeTexture, Category = "HoldoutComposite")
 	TObjectPtr<UTexture> CompositeTexture;
 
 public:
-	/** Directly sample the specified composite texture in screen space for the final post-processing composite, instead of using the separate built-in custom render pass. */
-	UPROPERTY(EditAnywhere, BlueprintReadWrite, Category = "HoldoutComposite", meta = (EditCondition = "CompositeTexture != nullptr"))
-	bool bCompositeTextureOverrideRender = false;
 
 	/** Activate projection of the composite texture onto the composite mesh actors. Continuously updates four material parameters called "XPlane", "YPlane", "ZPlane" & "WPlane" representing the view-projection matrix. */
 	UPROPERTY(EditAnywhere, BlueprintReadWrite, Category = "HoldoutComposite")
@@ -95,14 +99,29 @@ public:
 	UPROPERTY(EditAnywhere, BlueprintReadWrite, Category = "HoldoutComposite")
 	TWeakObjectPtr<ACameraActor> ProjectionCamera;
 
-	/** Enable local settings override.*/
-	UPROPERTY(EditAnywhere, BlueprintGetter = GetEnableOverrideSettings, BlueprintSetter = SetEnableOverrideSettings, Category = "HoldoutComposite")
-	bool bEnableOverrideSettings = false;
+	/* Passes to pre-process the specified composite texture inputs. */
+	UPROPERTY(EditAnywhere, BlueprintReadWrite, Category = "HoldoutComposite", meta = (DisplayName = "Composite Processing"))
+	TArray<FHoldoutCompositeProcessingPasses> CompositeProcessing;
+
+	/* Post-processing passes to update the engine SSR (NOTE: Runs before any existing post-process material at the same location). */
+	UPROPERTY(EditAnywhere, BlueprintReadWrite, Instanced, Category = "HoldoutComposite|PostProcessing", meta = (DisplayName = "SSR Input"))
+	TArray<TObjectPtr<UHoldoutCompositePassBase>> Passes_SSRInput;
+
+	/* Post-processing passes to update scene color after tonemap (NOTE: Runs before any existing post-process material at the same location). */
+	UPROPERTY(EditAnywhere, BlueprintReadWrite, Instanced, Category = "HoldoutComposite|PostProcessing", meta = (DisplayName = "After Tone Map"))
+	TArray<TObjectPtr<UHoldoutCompositePassBase>> Passes_AfterToneMap;
+
+	/* Post-processing passes to update scene color after all of post-processing. */
+	UPROPERTY(EditAnywhere, BlueprintReadWrite, Instanced, Category = "HoldoutComposite|PostProcessing", meta = (DisplayName = "After Post-Processing"))
+	TArray<TObjectPtr<UHoldoutCompositePassBase>> Passes_AfterPostProcessing;
 
-	/** Post process settings to use for this holdout composite. */
-	UPROPERTY(EditAnywhere, BlueprintReadOnly, Category = "HoldoutComposite", meta = (EditCondition = bOverrideDefaultSettings, DisplayAfter = "bOverrideDefaultSettings"))
-	FHoldoutCompositeSettings OverrideSettings = {};
+	/* Override default view user flags. */
+	UPROPERTY(EditAnywhere, BlueprintReadWrite, Category = "HoldoutComposite", AdvancedDisplay)
+	bool bOverridesViewUserFlags;
 
+	/* Custom user flags value used to alter materials in the composite render pass. */
+	UPROPERTY(EditAnywhere, BlueprintReadWrite, Category = "HoldoutComposite", AdvancedDisplay, meta = (EditCondition = "bOverridesViewUserFlags"))
+	int32 ViewUserFlags;
 private:
 
 	/* Returns true if the holdout composite is active. */
diff --git a/Engine/Plugins/Experimental/Compositing/HoldoutComposite/Source/HoldoutComposite/Public/HoldoutCompositeSettings.h b/Engine/Plugins/Experimental/Compositing/HoldoutComposite/Source/HoldoutComposite/Public/HoldoutCompositeSettings.h
--- a/Engine/Plugins/Experimental/Compositing/HoldoutComposite/Source/HoldoutComposite/Public/HoldoutCompositeSettings.h
+++ b/Engine/Plugins/Experimental/Compositing/HoldoutComposite/Source/HoldoutComposite/Public/HoldoutCompositeSettings.h
@@ -9,10 +9,6 @@
 
 #include "HoldoutCompositeSettings.generated.h"
 
-#if WITH_EDITOR
-DECLARE_MULTICAST_DELEGATE(FOnDefaultCompositeSettingsChanged);
-#endif
-
 class UPrimitiveComponent;
 
 namespace UE
@@ -26,66 +22,6 @@ namespace UE
 	}
 }
 
-/* Post-processing location of the main holdout composite pass. */
-UENUM()
-enum class EHoldoutCompositePassLocation
-{
-	AfterTonemap,
-	EndOfPostProcessing
-};
-
-/**
- * Settings for the active composite (default values can be overriden by the composite actor).
- */
-USTRUCT(BlueprintType)
-struct HOLDOUTCOMPOSITE_API FHoldoutCompositeSettings
-{
-	GENERATED_BODY()
-
-	/** Constructor */
-	FHoldoutCompositeSettings();
-
-	/** When enabled, the view global exposure is applied onto the separate render when composited. */
-	UPROPERTY(config, EditAnywhere, Category = HoldoutComposite)
-	bool bFollowSceneExposure;
-
-	/** When enabled, the separate composited render is also used to update screen-space reflections. */
-	UPROPERTY(config, EditAnywhere, Category = HoldoutComposite)
-	bool bUpdateSSR;
-
-	/** Enables the development shader debug feature that routes the Base Color output to Emissive for the separate render. Non-shipping PC build only. */
-	UPROPERTY(config, EditAnywhere, Category = HoldoutComposite)
-	bool bEnableUnlitViewmode;
-
-	/** When enabled, the TSR lens distortion (r.TSR.LensDistortion / Scene View Extension (Experimental) on lens files) is also applied on the separate composited render. */
-	UPROPERTY(config, EditAnywhere, Category = HoldoutComposite, DisplayName = "Apply TSR Lens Distortion")
-	bool bApplyTSRLensDistortion;
-
-	/** When enabled, FXAA is applied on the separate composited render. */
-	UPROPERTY(config, EditAnywhere, Category = HoldoutComposite, DisplayName = "Apply FXAA")
-	bool bApplyFXAA;
-
-	/** Composite (scene view extension) pass location, which defaults to "After Tonemap". "End Of Post Processing" may be used to push the composite pass after additional post-process materials (such as the regular lens distortion) and OCIO. */
-	UPROPERTY(config, EditAnywhere, Category = HoldoutComposite, AdvancedDisplay)
-	EHoldoutCompositePassLocation CompositePassLocation;
-
-	/** Composite (scene view extension) pass priority, which defaults to before OpenColorIO. */
-	UPROPERTY(config, EditAnywhere, Category = HoldoutComposite, AdvancedDisplay)
-	int32 CompositePassPriority;
-
-	/** Override the composite view user flags with the value specified below. */
-	UPROPERTY(config, EditAnywhere, Category = HoldoutComposite, AdvancedDisplay)
-	bool bOverridesViewUserFlags;
-
-	/** User flag for the composite render, allowing per-view material overrides. Active if greater or equal to zero. */
-	UPROPERTY(config, EditAnywhere, Category = HoldoutComposite, AdvancedDisplay, meta = (EditCondition = "bOverridesViewUserFlags"))
-	int32 CompositeViewUserFlags;
-
-	/** If enabled, composite renders in linear colorspace. */
-	UPROPERTY(config)
-	bool bLinearComposite;
-};
-
 /**
  * Settings for the HoldoutComposite module.
  */
@@ -114,17 +50,17 @@ public:
 	bool IsAllowedPrimitiveClass(const UPrimitiveComponent* InPrimitiveComponent) const;
 
 	/** Primitive component classes that do not support the holdout composite.*/
-	UPROPERTY(config, EditAnywhere, Category = General)
+	UPROPERTY(config, EditAnywhere, Category = HoldoutComposite)
 	TArray<FSoftClassPath> DisabledPrimitiveClasses;
 
-	/** Default holdout composite settings.*/
-	UPROPERTY(config, EditAnywhere, Category = General, meta = (ShowOnlyInnerProperties))
-	FHoldoutCompositeSettings DefaultCompositeSettings;
+	/** Enables the development shader debug feature that routes the Base Color output to Emissive for the separate render. Non-shipping PC build only. */
+	UPROPERTY(config, EditAnywhere, Category = HoldoutComposite)
+	bool bEnableUnlitViewmode;
+
+	/** Composite (scene view extension) post-processing priority, which defaults to before OpenColorIO. */
+	UPROPERTY(config, EditAnywhere, Category = HoldoutComposite, AdvancedDisplay)
+	int32 SceneViewExtensionPriority;
 
-#if WITH_EDITOR
-	/** Delegate event that triggers when the default settings are changed.*/
-	FOnDefaultCompositeSettingsChanged OnDefaultCompositeSettingsChanged;
-#endif
 private:
 
 	/** Cache disable class pointers to prevent load object calls during allowed class checks. */
diff --git a/Engine/Plugins/Experimental/Compositing/HoldoutComposite/Source/HoldoutComposite/Public/HoldoutCompositeSubsystem.h b/Engine/Plugins/Experimental/Compositing/HoldoutComposite/Source/HoldoutComposite/Public/HoldoutCompositeSubsystem.h
--- a/Engine/Plugins/Experimental/Compositing/HoldoutComposite/Source/HoldoutComposite/Public/HoldoutCompositeSubsystem.h
+++ b/Engine/Plugins/Experimental/Compositing/HoldoutComposite/Source/HoldoutComposite/Public/HoldoutCompositeSubsystem.h
@@ -7,17 +7,20 @@
 
 #include "HoldoutCompositeSubsystem.generated.h"
 
+class AHoldoutComposite;
 class FHoldoutCompositeSceneViewExtension;
 struct FHoldoutCompositeSettings;
 class SNotificationItem;
 class UPrimitiveComponent;
 class URendererSettings;
 
+class UHoldoutCompositePassBase;
+
 /**
  * Composite subsytem used as an interface to the (private) scene view extension.
  */
 UCLASS(BlueprintType, Transient)
-class HOLDOUTCOMPOSITE_API UHoldoutCompositeSubsystem : public UWorldSubsystem
+class HOLDOUTCOMPOSITE_API UHoldoutCompositeSubsystem : public UTickableWorldSubsystem
 {
 	GENERATED_BODY()
 
@@ -29,35 +32,32 @@ public:
 	virtual void Deinitialize() override;
 	// USubsystem implementation End
 
-	/* Register a single primitive for compositing. */
-	UFUNCTION(BlueprintCallable, Category = "Holdout Composite")
-	virtual void RegisterPrimitive(UPrimitiveComponent* InPrimitiveComponent);
+	// FTickableGameObject implementation Begin
+	virtual bool IsTickableInEditor() const override { return true; }
+	virtual void Tick(float DeltaTime) override;
+	// FTickableGameObject implementation End
 
-	/* Register multiple primitives for compositing. */
-	virtual void RegisterPrimitives(const TArray<UPrimitiveComponent*>& InPrimitiveComponents);
+	virtual TStatId GetStatId() const override { RETURN_QUICK_DECLARE_CYCLE_STAT(UHoldoutCompositeSubsystem, STATGROUP_Tickables); }
 
-	/* Unregister a single primitive from compositing. */
-	UFUNCTION(BlueprintCallable, Category = "Holdout Composite")
-	virtual void UnregisterPrimitive(UPrimitiveComponent* InPrimitiveComponent);
+	/* Register the composite actor for full control over the compositing process. */
+	void RegisterActor(AHoldoutComposite* InActor);
 
-	/* Unregister multiple primitives from compositing. */
-	virtual void UnregisterPrimitives(const TArray<UPrimitiveComponent*>& InPrimitiveComponents);
+	/* Unregister the composite actor. */
+	void ResetActor();
 
-	/* Set additional post-process materials to be applied on the separate composited render. */
+	/* Register a single primitive for compositing. */
 	UFUNCTION(BlueprintCallable, Category = "Holdout Composite")
-	virtual void SetCompositePostProcessMaterials(const TArray<UMaterialInterface*>& InMaterials);
+	void RegisterPrimitive(UPrimitiveComponent* InPrimitiveComponent);
 
-	/* Remove all post-process materials to be applied on the separate composited render. */
-	UFUNCTION(BlueprintCallable, Category = "Holdout Composite")
-	virtual void ResetCompositePostProcessMaterials();
+	/* Register multiple primitives for compositing. */
+	void RegisterPrimitives(const TArray<UPrimitiveComponent*>& InPrimitiveComponents);
 
-	/* Override project default composite settings for the active holdout composite passes. */
+	/* Unregister a single primitive from compositing. */
 	UFUNCTION(BlueprintCallable, Category = "Holdout Composite")
-	virtual void OverrideCompositeSettings(const FHoldoutCompositeSettings& InSettings, int32 InPriority);
+	void UnregisterPrimitive(UPrimitiveComponent* InPrimitiveComponent);
 
-	/* Revert composite settings back to project defaults. */
-	UFUNCTION(BlueprintCallable, Category = "Holdout Composite")
-	virtual void ResetCompositeSettings();
+	/* Unregister multiple primitives from compositing. */
+	void UnregisterPrimitives(const TArray<UPrimitiveComponent*>& InPrimitiveComponents);
 
 	/* Set a texture override to be used instead of the separate composited render. */
 	UFUNCTION(BlueprintCallable, Category = "Holdout Composite")
@@ -66,10 +66,6 @@ public:
 	/** True if the project settings are valid for the HoldoutComposite plugin to work. */
 	static bool IsProjectSettingsValid();
 
-protected:
-	/** Inherited classes can implement notification in a different way. */
-	virtual bool EnableHoldoutSettingsNotification() const;
-
 private:
 
 	/* Returns true if the (renderer) project settings are correctly enabled for the composite to be active. */
@@ -81,8 +77,14 @@ private:
 
 	/* Toast notification item. */
 	TWeakPtr<SNotificationItem> HoldoutNotificationItem;
+
+	/* Composite actor count to ensure no more than one actor is used. */
+	int32 CompositeActorCount = 0;
 #endif
 
+	/* Registered composite actor. */
+	TWeakObjectPtr<AHoldoutComposite> CompositeActorPtr;
+
 	/* Owned scene view extension. */
 	TSharedPtr<FHoldoutCompositeSceneViewExtension, ESPMode::ThreadSafe> HoldoutCompositeViewExtension;
 };
diff --git a/Engine/Plugins/Experimental/Compositing/HoldoutComposite/Source/HoldoutComposite/Public/Passes/HoldoutCompositePassBase.h b/Engine/Plugins/Experimental/Compositing/HoldoutComposite/Source/HoldoutComposite/Public/Passes/HoldoutCompositePassBase.h
new file mode 100644
--- /dev/null
+++ b/Engine/Plugins/Experimental/Compositing/HoldoutComposite/Source/HoldoutComposite/Public/Passes/HoldoutCompositePassBase.h
@@ -0,0 +1,139 @@
+// Copyright Epic Games, Inc. All Rights Reserved.
+
+#pragma once
+
+#include "CoreTypes.h"
+#include "Containers/ContainersFwd.h"
+#include "SceneViewExtension.h"
+#include "UObject/Object.h"
+
+#include "HoldoutCompositePassBase.generated.h"
+
+class AHoldoutComposite;
+class FRDGBuilder;
+class FSceneView;
+struct FScreenPassTexture;
+struct FPostProcessMaterialInputs;
+class IHoldoutCompositePassProxy;
+class UHoldoutCompositePassBase;
+
+/* Composite texture input type. */
+UENUM()
+enum class EHoldoutCompositePassInput : uint8
+{
+	/** Default input (usually scene color), or result of the previous pass. */
+	Default,
+	/** Texture from the built-in separate composite custom render pass of the specified primitives. */
+	CompositeRenderPass,
+	/** Direct composite texture, usually a media texture. */
+	CompositeTexture
+};
+
+namespace UE
+{
+	namespace HoldoutComposite
+	{
+		namespace Private
+		{
+			/** Internal render-thread struct for rendering work per frame. */
+			struct FRenderWork final
+			{
+				/** Constructor. */
+				FRenderWork();
+
+				/** Default frame work. */
+				static const FRenderWork& GetDefault();
+
+				/** Add input (pre) processing work. */
+				void AddCompositeProcessing(EHoldoutCompositePassInput InputType, TArrayView<TObjectPtr<UHoldoutCompositePassBase>> InPasses);
+
+				/** Add post-processing work. */
+				void AddPostProcessing(ISceneViewExtension::EPostProcessingPass InLocation, TArrayView<TObjectPtr<UHoldoutCompositePassBase>> InPasses);
+
+				/** Input (pre) processing passes. */
+				TArray<TPair<EHoldoutCompositePassInput, TArray<const IHoldoutCompositePassProxy*>>> InputProcessing;
+
+				/** Post-processing passes at the specified locations. */
+				TSortedMap<ISceneViewExtension::EPostProcessingPass, TArray<const IHoldoutCompositePassProxy*>> FramePasses;
+
+			private:
+				/** Proxy allocator. */
+				TUniquePtr<FSceneRenderingBulkObjectAllocator> FrameAllocator;
+
+				/** Convenience function for pass registration. */
+				void RegisterPasses(TArrayView<TObjectPtr<UHoldoutCompositePassBase>> InPasses, TArray<const IHoldoutCompositePassProxy*>& OutProxies);
+			};
+		}
+	}
+}
+
+/** Render-thread pass proxy. */
+class HOLDOUTCOMPOSITE_API IHoldoutCompositePassProxy
+{
+public:
+	/** Constructor. */
+	IHoldoutCompositePassProxy(EHoldoutCompositePassInput Input0)
+		: InputTypes({ Input0 })
+	{ };
+
+	/* Virtual destructor in case deletion is done on interface pointers. */
+	virtual ~IHoldoutCompositePassProxy() { };
+
+	/* Render-thread add pass method to override. */
+	virtual FScreenPassTexture Add(FRDGBuilder& GraphBuilder, const FSceneView& InView, const FPostProcessMaterialInputs& Inputs, ISceneViewExtension::EPostProcessingPass InLocation) const = 0;
+
+	/* Number of inputs used by the pass. */
+	int32 GetNumInputs() const { return InputTypes.Num(); }
+
+	/* Number of inputs used by the pass. */
+	EHoldoutCompositePassInput GetPassInput(int32 Index) const { return (Index < GetNumInputs()) ? InputTypes[Index] : EHoldoutCompositePassInput::Default; };
+
+protected:
+	/* List of pass input types. */
+	TArray<EHoldoutCompositePassInput> InputTypes;
+};
+
+/* Base class for defining a pass object, and its associated render proxy. */
+UCLASS(Abstract)
+class HOLDOUTCOMPOSITE_API UHoldoutCompositePassBase : public UObject
+{
+	GENERATED_BODY()
+
+public:
+	/** Constructor. */
+	UHoldoutCompositePassBase(const FObjectInitializer& ObjectInitializer);
+	/** Desstructor. */
+	virtual ~UHoldoutCompositePassBase();
+
+#if WITH_EDITOR
+	virtual bool CanEditChange(const FProperty* InProperty) const;
+
+	/* Set editor-only state to track whether inputs can be edited. */
+	void SetCanEditInput(bool bInEdit);
+#endif
+
+	/** Getter function to override, returning pass proxies to be passed to the render thread. (Proxy objects should be allocated from the provided allocator.) */
+	virtual bool GetProxy(FSceneRenderingBulkObjectAllocator& InFrameAllocator, IHoldoutCompositePassProxy*& OutProxy) const { return false; };
+
+	/* Whether or not the holdout composite is active. */
+	UPROPERTY(EditAnywhere, BlueprintReadWrite, Category = "", meta = (DisplayPriority = "100"))
+	bool bIsEnabled = true;
+
+	/* First pass input type. */
+	UPROPERTY(EditAnywhere, BlueprintReadWrite, Category = "", meta = (DisplayAfter = "bIsEnabled", GetRestrictedEnumValues = "GetRestrictedInputs"))
+	EHoldoutCompositePassInput Input0 = EHoldoutCompositePassInput::Default;
+
+protected:
+#if WITH_EDITOR
+	/* Convenience function for restricting composite texture from inputs if it is not present. */
+	UFUNCTION()
+	TArray<FString> GetRestrictedInputs() const;
+
+	/* Editor-only state to track whether inputs can be edited. */
+	bool bCanEditInput = true;
+#endif
+
+	/** Cached pointer to owning actor */
+	const AHoldoutComposite* OwnerPrivate;
+};
+
diff --git a/Engine/Plugins/Experimental/Compositing/HoldoutComposite/Source/HoldoutComposite/Public/Passes/HoldoutCompositePassFXAA.h b/Engine/Plugins/Experimental/Compositing/HoldoutComposite/Source/HoldoutComposite/Public/Passes/HoldoutCompositePassFXAA.h
new file mode 100644
--- /dev/null
+++ b/Engine/Plugins/Experimental/Compositing/HoldoutComposite/Source/HoldoutComposite/Public/Passes/HoldoutCompositePassFXAA.h
@@ -0,0 +1,32 @@
+// Copyright Epic Games, Inc. All Rights Reserved.
+
+#pragma once
+
+#include "HoldoutCompositePassBase.h"
+
+#include "HoldoutCompositePassFXAA.generated.h"
+
+/* FXAA holdout composite pass, mostly useful for applying on the composite render pass for CG/Motion Graphics. */
+UCLASS(BlueprintType, Blueprintable, editinlinenew, CollapseCategories, meta=(DisplayName="FXAA Pass"))
+class HOLDOUTCOMPOSITE_API UHoldoutCompositeFXAAPass : public UHoldoutCompositePassBase
+{
+	GENERATED_BODY()
+
+public:
+	/* Constructor */
+	UHoldoutCompositeFXAAPass(const FObjectInitializer& ObjectInitializer);
+	
+	/* Destructor */
+	~UHoldoutCompositeFXAAPass();
+
+	/* Toggle to ensure FXAA always gammatizes colors before processing, if necessary. */
+	UPROPERTY(EditAnywhere, BlueprintReadWrite, Category = "")
+	bool bLinearSourceColors;
+
+	/* FXAA quality control, matching the "r.fxaa.quality" console variable values. */
+	UPROPERTY(EditAnywhere, BlueprintReadWrite, Category = "", meta=(ClampMin="0", ClampMax="5", UIMin = "0", UIMax = "5"))
+	int32 Quality;
+
+	virtual bool GetProxy(FSceneRenderingBulkObjectAllocator& InFrameAllocator, IHoldoutCompositePassProxy*& OutProxy) const override;
+};
+
diff --git a/Engine/Plugins/Experimental/Compositing/HoldoutComposite/Source/HoldoutComposite/Public/Passes/HoldoutCompositePassMaterial.h b/Engine/Plugins/Experimental/Compositing/HoldoutComposite/Source/HoldoutComposite/Public/Passes/HoldoutCompositePassMaterial.h
new file mode 100644
--- /dev/null
+++ b/Engine/Plugins/Experimental/Compositing/HoldoutComposite/Source/HoldoutComposite/Public/Passes/HoldoutCompositePassMaterial.h
@@ -0,0 +1,30 @@
+// Copyright Epic Games, Inc. All Rights Reserved.
+
+#pragma once
+
+#include "HoldoutCompositePassBase.h"
+
+#include "HoldoutCompositePassMaterial.generated.h"
+
+class UMaterialInterface;
+
+/* Post-process material holdout composite pass. */
+UCLASS(BlueprintType, Blueprintable, editinlinenew, CollapseCategories, meta = (DisplayName = "Post-Process Material Pass"))
+class HOLDOUTCOMPOSITE_API UHoldoutCompositePassMaterial : public UHoldoutCompositePassBase
+{
+	GENERATED_BODY()
+
+public:
+	/* Constructor */
+	UHoldoutCompositePassMaterial(const FObjectInitializer& ObjectInitializer);
+
+	/* Destructor */
+	~UHoldoutCompositePassMaterial();
+
+	/* Post-process material to execute. Input0 is connected to SceneTexture's PostprocessInput0. */
+	UPROPERTY(EditAnywhere, BlueprintReadWrite, Category = "")
+	TObjectPtr<UMaterialInterface> PostProcessMaterial;
+
+	virtual bool GetProxy(FSceneRenderingBulkObjectAllocator& InFrameAllocator, IHoldoutCompositePassProxy*& OutProxy) const override;
+};
+
diff --git a/Engine/Plugins/Experimental/Compositing/HoldoutComposite/Source/HoldoutComposite/Public/Passes/HoldoutCompositePassMerge.h b/Engine/Plugins/Experimental/Compositing/HoldoutComposite/Source/HoldoutComposite/Public/Passes/HoldoutCompositePassMerge.h
new file mode 100644
--- /dev/null
+++ b/Engine/Plugins/Experimental/Compositing/HoldoutComposite/Source/HoldoutComposite/Public/Passes/HoldoutCompositePassMerge.h
@@ -0,0 +1,70 @@
+// Copyright Epic Games, Inc. All Rights Reserved.
+
+#pragma once
+
+#include "CoreTypes.h"
+#include "HoldoutCompositePassBase.h"
+
+#include "HoldoutCompositePassMerge.generated.h"
+
+UENUM()
+enum class EHoldoutCompositePassMergeOp : uint8
+{
+	Over,
+	Add,
+};
+
+namespace UE
+{
+	namespace HoldoutComposite
+	{
+		namespace Private
+		{
+			class FHoldoutCompositePassMergeProxy : public IHoldoutCompositePassProxy
+			{
+			public:
+				/** Constructor. */
+				FHoldoutCompositePassMergeProxy(EHoldoutCompositePassInput Input0, EHoldoutCompositePassInput Input1);
+
+				FScreenPassTexture Add(FRDGBuilder& GraphBuilder, const FSceneView& InView, const FPostProcessMaterialInputs& Inputs, ISceneViewExtension::EPostProcessingPass InLocation) const override;
+
+				/** Merge operation. */
+				EHoldoutCompositePassMergeOp MergeOp = EHoldoutCompositePassMergeOp::Over;
+
+				/* Unpremultiplication required for recovering opaque colors for the default over blend. */
+				bool bUnpremulInput1 = true;
+
+				/* Output is assumed to be scene color with inverted alpha. */
+				bool bOutputSceneColor = true;
+			};
+		}
+	}
+}
+
+/* Pass to merge Input0 with Input1. */
+UCLASS(BlueprintType, Blueprintable, editinlinenew, CollapseCategories, meta = (DisplayName = "Merge Pass"))
+class HOLDOUTCOMPOSITE_API UHoldoutCompositePassMerge : public UHoldoutCompositePassBase
+{
+	GENERATED_BODY()
+
+public:
+	/* Constructor */
+	UHoldoutCompositePassMerge(const FObjectInitializer& ObjectInitializer);
+	/* Destructor */
+	~UHoldoutCompositePassMerge();
+
+	/* Second pass input type. */
+	UPROPERTY(EditAnywhere, BlueprintReadWrite, Category = "", meta = (DisplayAfter = "Input0", GetRestrictedEnumValues = "GetRestrictedInputs"))
+	EHoldoutCompositePassInput Input1 = EHoldoutCompositePassInput::CompositeRenderPass;
+
+	/* Merge operation applied on Input0 with Input1. */
+	UPROPERTY(EditAnywhere, BlueprintReadWrite, Category = "")
+	EHoldoutCompositePassMergeOp MergeOp = EHoldoutCompositePassMergeOp::Over;
+
+	/* Unpremultiplication required for recovering opaque colors for the default over blend. */
+	UPROPERTY(EditAnywhere, BlueprintReadWrite, Category = "")
+	bool bUnpremulInput1 = true;
+
+	virtual bool GetProxy(FSceneRenderingBulkObjectAllocator& InFrameAllocator, IHoldoutCompositePassProxy*& OutProxy) const override;
+};
+
