HoldoutComposite: Re-working of the camera projection feature into a separate component, similar to the lens component, whose sole responsibility is to update an MPC view-projection matrix with the target camera.

#virtualized

diff --git a/Engine/Plugins/Experimental/Compositing/HoldoutComposite/Source/HoldoutComposite/Private/HoldoutComposite.cpp b/Engine/Plugins/Experimental/Compositing/HoldoutComposite/Source/HoldoutComposite/Private/HoldoutComposite.cpp
--- a/Engine/Plugins/Experimental/Compositing/HoldoutComposite/Source/HoldoutComposite/Private/HoldoutComposite.cpp
+++ b/Engine/Plugins/Experimental/Compositing/HoldoutComposite/Source/HoldoutComposite/Private/HoldoutComposite.cpp
@@ -97,16 +97,6 @@ void AHoldoutComposite::PostRegisterAllComponents()
 	{
 		ApplyAllState();
 	}
-
-	/*
-	* NOTE: We must (unfortunately) rely on this after-tick event to update the view matrix projection,
-	* since viewport camera transform updates after ticking all world objects.
-	*
-	* The CompositePlane plugin updated this information on tick, but a lag becomes visible when
-	* piloting the source projection camera for example.
-	*/
-	OnWorldPreSendAllEndOfFrameUpdatesHandle = FWorldDelegates::OnWorldPreSendAllEndOfFrameUpdates.AddUObject(this, &AHoldoutComposite::EndOfFrameUpdate);
-
 	UHoldoutCompositeSubsystem* Subsystem = UWorld::GetSubsystem<UHoldoutCompositeSubsystem>(GetWorld());
 	if (IsValid(Subsystem))
 	{
@@ -124,8 +114,6 @@ void AHoldoutComposite::PostUnregisterAllComponents()
 		Subsystem->ResetActor();
 	}
 
-	FWorldDelegates::OnWorldPreSendAllEndOfFrameUpdates.Remove(OnWorldPreSendAllEndOfFrameUpdatesHandle);
-
 	RemoveAllState();
 
 	Super::PostUnregisterAllComponents();
@@ -152,104 +140,6 @@ void AHoldoutComposite::Tick(float DeltaTime)
 	}
 }
 
-void AHoldoutComposite::EndOfFrameUpdate(UWorld* InWorld)
-{
-	SCOPE_CYCLE_COUNTER(STAT_HoldoutComposite_Actor_EndOfFrameUpdate);
-
-	if (!IsActive())
-	{
-		return;
-	}
-
-	// Only update materials once per frame
-	if (PreviousFrameCounter == GFrameCounter)
-	{
-		return;
-	}
-	else
-	{
-		PreviousFrameCounter = GFrameCounter;
-	}
-
-	if (bCompositeTextureProjection)
-	{
-		bool bUpdateViewProjection = false;
-
-		FMinimalViewInfo DesiredView;
-		FMatrix ViewMatrix, ProjectionMatrix, ViewProjectionMatrix;
-
-		TStrongObjectPtr<ACameraActor> CameraActor = ProjectionCamera.Pin();
-		if (CameraActor.IsValid())
-		{
-			if (UCameraComponent* CameraComponent = CameraActor->GetCameraComponent())
-			{
-				CameraComponent->GetCameraView(FApp::GetDeltaTime(), DesiredView);
-
-				UGameplayStatics::GetViewProjectionMatrix(DesiredView, ViewMatrix, ProjectionMatrix, ViewProjectionMatrix);
-				bUpdateViewProjection = true;
-			}
-		}
-		else if (const APlayerController* PlayerController = GetWorld()->GetFirstPlayerController())
-		{
-			if (AActor* ViewTarget = PlayerController->GetViewTarget())
-			{
-				UGameplayStatics::CalculateViewProjectionMatricesFromViewTarget(ViewTarget, ViewMatrix, ProjectionMatrix, ViewProjectionMatrix);
-				bUpdateViewProjection = true;
-			}
-		}
-#if WITH_EDITOR
-		else if (GCurrentLevelEditingViewportClient)
-		{
-			if (UCameraComponent* CameraComponent = GCurrentLevelEditingViewportClient->GetCameraComponentForView())
-			{
-				CameraComponent->GetCameraView(FApp::GetDeltaTime(), DesiredView);
-
-				UGameplayStatics::GetViewProjectionMatrix(DesiredView, ViewMatrix, ProjectionMatrix, ViewProjectionMatrix);
-				bUpdateViewProjection = true;
-			}
-			else
-			{
-				static FSceneViewFamily::ConstructionValues CVS(nullptr, nullptr, FEngineShowFlags(EShowFlagInitMode::ESFIM_Game));
-				CVS.SetTime(FGameTime::GetTimeSinceAppStart());
-
-				FSceneViewFamilyContext ViewFamily(CVS);
-				FSceneView* View = GCurrentLevelEditingViewportClient->CalcSceneView(&ViewFamily);
-
-				ViewProjectionMatrix = View->ViewMatrices.GetViewProjectionMatrix();
-				bUpdateViewProjection = true;
-			}
-		}
-#endif
-
-#if 0
-		if (CachedViewProjectionMatrix.Equals(ViewProjectionMatrix, UE_SMALL_NUMBER))
-		{
-			bUpdateViewProjection = false;
-		}
-#endif
-
-		if (bUpdateViewProjection)
-		{
-			SetMaterialParameterOnCompositeMeshes(
-				[this, &ViewProjectionMatrix](UMaterialInstanceDynamic* MID)
-				{
-					const FVector4 XPlane{ ViewProjectionMatrix.M[0][0], ViewProjectionMatrix.M[0][1], ViewProjectionMatrix.M[0][2], ViewProjectionMatrix.M[0][3] };
-					const FVector4 YPlane{ ViewProjectionMatrix.M[1][0], ViewProjectionMatrix.M[1][1], ViewProjectionMatrix.M[1][2], ViewProjectionMatrix.M[1][3] };
-					const FVector4 ZPlane{ ViewProjectionMatrix.M[2][0], ViewProjectionMatrix.M[2][1], ViewProjectionMatrix.M[2][2], ViewProjectionMatrix.M[2][3] };
-					const FVector4 WPlane{ ViewProjectionMatrix.M[3][0], ViewProjectionMatrix.M[3][1], ViewProjectionMatrix.M[3][2], ViewProjectionMatrix.M[3][3] };
-
-					MID->SetVectorParameterValue(TEXT("XPlane"), XPlane);
-					MID->SetVectorParameterValue(TEXT("YPlane"), YPlane);
-					MID->SetVectorParameterValue(TEXT("ZPlane"), ZPlane);
-					MID->SetVectorParameterValue(TEXT("WPlane"), WPlane);
-				}
-			);
-
-			CachedViewProjectionMatrix = ViewProjectionMatrix;
-		}
-	}
-}
-
 void AHoldoutComposite::BeginDestroy()
 {
 	Super::BeginDestroy();
@@ -379,29 +269,3 @@ void AHoldoutComposite::RemoveAllState()
 	}
 }
 
-void AHoldoutComposite::SetMaterialParameterOnCompositeMeshes(TFunction<void(UMaterialInstanceDynamic* MID)> ApplyFn)
-{
-	using namespace UE::HoldoutComposite::Private;
-
-	TArray<UPrimitiveComponent*> Primitives = GetPrimitives(CompositeMeshActors);
-	
-	for (UPrimitiveComponent* Primitive : Primitives)
-	{
-		if (!IsValid(Primitive))
-		{
-			continue;
-		}
-
-		for (int32 ElementIndex = 0; ElementIndex < Primitive->GetNumMaterials(); ElementIndex++)
-		{
-			UMaterialInterface* Material = Primitive->GetMaterial(ElementIndex);
-			UMaterialInstanceDynamic* MID = Cast<UMaterialInstanceDynamic>(Material);
-			if (!MID)
-			{
-				MID = Primitive->CreateAndSetMaterialInstanceDynamicFromMaterial(ElementIndex, Material);
-			}
-
-			ApplyFn(MID);
-		}
-	}
-}
diff --git a/Engine/Plugins/Experimental/Compositing/HoldoutComposite/Source/HoldoutComposite/Private/HoldoutCompositeViewProjectionComponent.cpp b/Engine/Plugins/Experimental/Compositing/HoldoutComposite/Source/HoldoutComposite/Private/HoldoutCompositeViewProjectionComponent.cpp
new file mode 100644
--- /dev/null
+++ b/Engine/Plugins/Experimental/Compositing/HoldoutComposite/Source/HoldoutComposite/Private/HoldoutCompositeViewProjectionComponent.cpp
@@ -0,0 +1,185 @@
+// Copyright Epic Games, Inc. All Rights Reserved.
+
+#include "HoldoutCompositeViewProjectionComponent.h"
+
+#include "Camera/CameraComponent.h"
+#include "Engine/Engine.h"
+#include "Engine/World.h"
+#include "HoldoutCompositeSubsystem.h"
+#include "Kismet/GameplayStatics.h"
+#include "Materials/MaterialParameterCollection.h"
+#include "Materials/MaterialParameterCollectionInstance.h"
+
+#define LOCTEXT_NAMESPACE "HoldoutComposite"
+
+UHoldoutCompositeViewProjectionComponent::UHoldoutCompositeViewProjectionComponent(const FObjectInitializer& ObjectInitializer)
+	: Super(ObjectInitializer)
+{
+}
+
+void UHoldoutCompositeViewProjectionComponent::ForceUpdate()
+{
+	LastViewProjectionMatrix = FMatrix::Identity;
+}
+
+void UHoldoutCompositeViewProjectionComponent::InitDefaultCamera()
+{
+	UCameraComponent* TargetCamera = Cast<UCameraComponent>(TargetCameraComponent.GetComponent(GetOwner()));
+	
+	if (!IsValid(TargetCamera))
+	{
+		// Find all CineCameraComponents on the same actor as this component and set the first one to be the target
+		TInlineComponentArray<UCameraComponent*> CameraComponents;
+		if (AActor* Owner = GetOwner())
+		{
+			Owner->GetComponents(CameraComponents);
+		}
+
+		if (CameraComponents.Num() > 0)
+		{
+			UCameraComponent* TargetComponent = CameraComponents[0];
+			TargetCameraComponent.ComponentProperty = TargetComponent->GetFName();
+			LastCameraComponent = TargetComponent;
+		}
+	}
+}
+
+void UHoldoutCompositeViewProjectionComponent::PostInitProperties()
+{
+	Super::PostInitProperties();
+
+#if WITH_EDITORONLY_DATA
+	if (!HasAnyFlags(RF_ClassDefaultObject | RF_NeedLoad))
+	{
+		MaterialParameterCollection = LoadObject<UMaterialParameterCollection>(nullptr, TEXT("/HoldoutComposite/Materials/MPC_HoldoutComposite.MPC_HoldoutComposite"));
+		MatrixParameterName = TEXT("CameraViewProjectionMatrix");
+	}
+#endif
+}
+
+void UHoldoutCompositeViewProjectionComponent::OnRegister()
+{
+	Super::OnRegister();
+
+	InitDefaultCamera();
+
+	/*
+	* NOTE: We must (unfortunately) rely on this after-tick event to update the view matrix projection,
+	* since viewport camera transform updates after ticking all world objects.
+	*
+	* The CompositePlane plugin updated this information on tick, but a lag becomes visible when
+	* piloting the source projection camera for example.
+	*/
+	OnWorldPreSendAllEndOfFrameUpdatesHandle = FWorldDelegates::OnWorldPreSendAllEndOfFrameUpdates.AddUObject(this, &UHoldoutCompositeViewProjectionComponent::UpdateProjection);
+}
+
+void UHoldoutCompositeViewProjectionComponent::PostDuplicate(bool bDuplicateForPIE)
+{
+	Super::PostDuplicate(bDuplicateForPIE);
+
+	InitDefaultCamera();
+}
+
+void UHoldoutCompositeViewProjectionComponent::PostEditImport()
+{
+	Super::PostEditImport();
+
+	InitDefaultCamera();
+}
+
+void UHoldoutCompositeViewProjectionComponent::PostLoad()
+{
+	Super::PostLoad();
+}
+
+void UHoldoutCompositeViewProjectionComponent::OnUnregister()
+{
+	FWorldDelegates::OnWorldPreSendAllEndOfFrameUpdates.Remove(OnWorldPreSendAllEndOfFrameUpdatesHandle);
+
+	Super::OnUnregister();
+}
+
+void UHoldoutCompositeViewProjectionComponent::UpdateProjection(UWorld* InWorld) const
+{
+	if (!IsValid(InWorld))
+	{
+		return;
+	}
+
+	if (!bIsEnabled)
+	{
+		return;
+	}
+
+	UMaterialParameterCollection* MPC = MaterialParameterCollection.Get();
+
+	if (IsValid(LastCameraComponent) && IsValid(MPC))
+	{
+		FMinimalViewInfo DesiredView;
+		FMatrix ViewMatrix, ProjectionMatrix, ViewProjectionMatrix;
+
+		LastCameraComponent->GetCameraView(FApp::GetDeltaTime(), DesiredView);
+
+		UGameplayStatics::GetViewProjectionMatrix(DesiredView, ViewMatrix, ProjectionMatrix, ViewProjectionMatrix);
+
+		if (!LastViewProjectionMatrix.Equals(ViewProjectionMatrix))
+		{
+			int32 ParameterProjection = INDEX_NONE;
+			for (int32 ParameterIndex = 0; ParameterIndex < MPC->VectorParameters.Num(); ParameterIndex++)
+			{
+				FName CurrentParameterName = MPC->VectorParameters[ParameterIndex].ParameterName;
+				if (!CurrentParameterName.IsNone())
+				{
+					if (CurrentParameterName == MatrixParameterName)
+					{
+						ParameterProjection = ParameterIndex;
+					}
+				}
+			}
+
+			UMaterialParameterCollectionInstance* MaterialParameterCollectionInstance = GetWorld()->GetParameterCollectionInstance(MaterialParameterCollection);
+			// Ensure there's space for 4 output vectors for the projection matrix
+			if (ParameterProjection != INDEX_NONE && ParameterProjection + 4 <= MaterialParameterCollection->VectorParameters.Num())
+			{
+				FMatrix44f ViewProjectionMatrixf = FMatrix44f(ViewProjectionMatrix);
+
+				// Store the vectors to the collection instance
+				const FLinearColor* MatrixVectors = (const FLinearColor*)&ViewProjectionMatrixf;
+				for (int32 ElementIndex = 0; ElementIndex < 4; ElementIndex++)
+				{
+					MaterialParameterCollectionInstance->SetVectorParameterValue(MaterialParameterCollection->VectorParameters[ParameterProjection + ElementIndex].ParameterName, MatrixVectors[ElementIndex]);
+				}
+
+				MaterialParameterCollectionInstance->UpdateRenderState(false);
+			}
+
+			LastViewProjectionMatrix = ViewProjectionMatrix;
+		}
+	}
+}
+
+void UHoldoutCompositeViewProjectionComponent::TickComponent(float DeltaTime, ELevelTick TickType, FActorComponentTickFunction* ThisTickFunction)
+{
+	Super::TickComponent(DeltaTime, TickType, ThisTickFunction);
+
+	// Note: Not active currently, since we use the end-of-frame callback.
+}
+
+#if WITH_EDITOR
+void UHoldoutCompositeViewProjectionComponent::PostEditChangeProperty(FPropertyChangedEvent& PropertyChangedEvent)
+{
+	const FName PropertyName = (PropertyChangedEvent.Property != NULL) ? PropertyChangedEvent.Property->GetFName() : NAME_None;
+
+	if (PropertyName == GET_MEMBER_NAME_CHECKED(ThisClass, TargetCameraComponent))
+	{
+		LastCameraComponent = Cast<UCameraComponent>(TargetCameraComponent.GetComponent(GetOwner()));
+	}
+
+	// Force MPC updates on any property change.
+	ForceUpdate();
+
+	Super::PostEditChangeProperty(PropertyChangedEvent);
+}
+#endif
+
+#undef LOCTEXT_NAMESPACE
diff --git a/Engine/Plugins/Experimental/Compositing/HoldoutComposite/Source/HoldoutComposite/Public/HoldoutComposite.h b/Engine/Plugins/Experimental/Compositing/HoldoutComposite/Source/HoldoutComposite/Public/HoldoutComposite.h
--- a/Engine/Plugins/Experimental/Compositing/HoldoutComposite/Source/HoldoutComposite/Public/HoldoutComposite.h
+++ b/Engine/Plugins/Experimental/Compositing/HoldoutComposite/Source/HoldoutComposite/Public/HoldoutComposite.h
@@ -91,14 +91,6 @@ private:
 
 public:
 
-	/** Activate projection of the composite texture onto the composite mesh actors. Continuously updates four material parameters called "XPlane", "YPlane", "ZPlane" & "WPlane" representing the view-projection matrix. */
-	UPROPERTY(EditAnywhere, BlueprintReadWrite, Category = "HoldoutComposite")
-	bool bCompositeTextureProjection = false;
-
-	/* Camera used for material projection. If none, the active viewport camera is used. */
-	UPROPERTY(EditAnywhere, BlueprintReadWrite, Category = "HoldoutComposite")
-	TWeakObjectPtr<ACameraActor> ProjectionCamera;
-
 	/* Passes to pre-process the specified composite texture inputs. */
 	UPROPERTY(EditAnywhere, BlueprintReadWrite, Category = "HoldoutComposite", meta = (DisplayName = "Composite Processing"))
 	TArray<FHoldoutCompositeProcessingPasses> CompositeProcessing;
@@ -133,19 +125,7 @@ private:
 	/* Remove all actor state from the holdout composite subsystem. */
 	void RemoveAllState();
 
-	/* End-of-frame update, primarily used for updating MID parameters on all composite mesh primitives (after viewport ticks). */
-	void EndOfFrameUpdate(UWorld* InWorld);
-
-	/* Convenience function to visit all composite meshes and update parameters on their MIDs. */
-	void SetMaterialParameterOnCompositeMeshes(TFunction<void(UMaterialInstanceDynamic* MID)> ApplyFn);
-
-	/* Cached view projection matrix to avoid unecessary reprojections. */
-	FMatrix CachedViewProjectionMatrix = FMatrix::Identity;
-
 	/* Cached previous frame counter to avoid double updates on the same frame. */
 	uint64 PreviousFrameCounter = 0;
-
-	/* Post-viewport update callback. */
-	FDelegateHandle OnWorldPreSendAllEndOfFrameUpdatesHandle;
 };
 
diff --git a/Engine/Plugins/Experimental/Compositing/HoldoutComposite/Source/HoldoutComposite/Public/HoldoutCompositeViewProjectionComponent.h b/Engine/Plugins/Experimental/Compositing/HoldoutComposite/Source/HoldoutComposite/Public/HoldoutCompositeViewProjectionComponent.h
new file mode 100644
--- /dev/null
+++ b/Engine/Plugins/Experimental/Compositing/HoldoutComposite/Source/HoldoutComposite/Public/HoldoutCompositeViewProjectionComponent.h
@@ -0,0 +1,74 @@
+// Copyright Epic Games, Inc. All Rights Reserved.
+
+#pragma once
+
+#include "CoreTypes.h"
+#include "Components/SceneComponent.h"
+
+#include "HoldoutCompositeViewProjectionComponent.generated.h"
+
+class UCameraComponent;
+class UMaterialParameterCollection;
+class UWorld;
+
+UCLASS(HideCategories = (Activation, Transform, Lighting, Rendering, Tags, Cooking, Physics, LOD, AssetUserData, Navigation), editinlinenew, meta = (BlueprintSpawnableComponent, DisplayName = "Camera View Projection"))
+class HOLDOUTCOMPOSITE_API UHoldoutCompositeViewProjectionComponent : public UActorComponent
+{
+	GENERATED_UCLASS_BODY()
+
+public:
+	//~ Begin UObject interface
+	virtual void PostInitProperties() override;
+	virtual void PostDuplicate(bool bDuplicateForPIE) override;
+	virtual void PostEditImport() override;
+	virtual void PostLoad() override;
+	//~ End UObject interface
+
+	//~ Begin UActorComponent interface
+	virtual void OnRegister() override;
+	virtual void OnUnregister() override;
+	virtual void TickComponent(float DeltaTime, enum ELevelTick TickType, FActorComponentTickFunction* ThisTickFunction) override;
+	//~ End UActorComponent interface
+
+#if WITH_EDITOR
+	virtual void PostEditChangeProperty(struct FPropertyChangedEvent& PropertyChangedEvent) override;
+#endif //WITH_EDITOR
+
+	/* Whether or not the component activates the view-projection matrix Material Parameter Collection update. */
+	UPROPERTY(EditAnywhere, BlueprintReadWrite, Category = "CameraProjection")
+	bool bIsEnabled = true;
+
+	/** The Material Parameter Collection in which the view-projection matrix should be stored. */
+	UPROPERTY(EditAnywhere, BlueprintReadWrite, Category = "CameraProjection")
+	TObjectPtr<UMaterialParameterCollection> MaterialParameterCollection;
+
+	/** Parameter name of the first element of the transform in the Material Parameter Collection set above.  Requires space for 4 vectors. */
+	UPROPERTY(EditAnywhere, BlueprintReadWrite, Category = "CameraProjection")
+	FName MatrixParameterName;
+
+	/** The CineCameraComponent whose view-projection matrix is used. */
+	UPROPERTY(EditInstanceOnly, AdvancedDisplay, Category = "CameraProjection", meta = (UseComponentPicker, AllowedClasses = "/Script/Engine.CameraComponent, /Script/CinematicCamera.CineCameraComponent"))
+	FComponentReference TargetCameraComponent;
+
+	/** Force update the Material Parameter Collection this frame. */
+	UFUNCTION(BlueprintCallable, Category = "CameraProjection")
+	void ForceUpdate();
+
+private:
+	/** If TargetCameraComponent is not set, initialize it to the first CineCameraComponent on the same actor as this component */
+	void InitDefaultCamera();
+	
+	/** Update the Material Parameter Collection with the camera view transform information. */
+	void UpdateProjection(UWorld* InWorld) const;
+
+	/** Cached most recent target camera, used to clean up the old camera when the user changes the target */
+	UPROPERTY()
+	TObjectPtr<UCameraComponent> LastCameraComponent = nullptr;
+
+	/** Cached view projection matrix, to avoid needless MPC updates. */
+	mutable FMatrix LastViewProjectionMatrix = FMatrix::Identity;
+
+	/* Post-viewport update callback. */
+	FDelegateHandle OnWorldPreSendAllEndOfFrameUpdatesHandle;
+};
+
