HoldoutComposite: Shader-typo hotfix for correct merge alpha output.

diff --git a/Engine/Plugins/Experimental/Compositing/HoldoutComposite/Shaders/Private/HoldoutCompositeMerge.usf b/Engine/Plugins/Experimental/Compositing/HoldoutComposite/Shaders/Private/HoldoutCompositeMerge.usf
--- a/Engine/Plugins/Experimental/Compositing/HoldoutComposite/Shaders/Private/HoldoutCompositeMerge.usf
+++ b/Engine/Plugins/Experimental/Compositing/HoldoutComposite/Shaders/Private/HoldoutCompositeMerge.usf
@@ -158,5 +158,5 @@ void MainPS(
 
     // Optionally re-invert alpha
     const bool bOutputSceneColor = (bFlags & 1);
-    OutColor.a = saturate((bOutputSceneColor > 0, 1.0 - OutColor.a, OutColor.a));
+    OutColor.a = saturate(select(bOutputSceneColor, 1.0 - OutColor.a, OutColor.a));
 }
