HoldoutComposite: Avoid registering any material that is not a post-process material.

diff --git a/Engine/Plugins/Experimental/Compositing/HoldoutComposite/Source/HoldoutComposite/Private/Passes/HoldoutCompositePassMaterial.cpp b/Engine/Plugins/Experimental/Compositing/HoldoutComposite/Source/HoldoutComposite/Private/Passes/HoldoutCompositePassMaterial.cpp
--- a/Engine/Plugins/Experimental/Compositing/HoldoutComposite/Source/HoldoutComposite/Private/Passes/HoldoutCompositePassMaterial.cpp
+++ b/Engine/Plugins/Experimental/Compositing/HoldoutComposite/Source/HoldoutComposite/Private/Passes/HoldoutCompositePassMaterial.cpp
@@ -45,7 +45,7 @@ UHoldoutCompositePassMaterial::~UHoldoutCompositePassMaterial() = default;
 
 bool UHoldoutCompositePassMaterial::GetProxy(FSceneRenderingBulkObjectAllocator& InFrameAllocator, IHoldoutCompositePassProxy*& OutProxy) const
 {
-	if (IsValid(PostProcessMaterial))
+	if (IsValid(PostProcessMaterial) && PostProcessMaterial->IsPostProcessMaterial())
 	{
 		FHoldoutCompositePassMaterialProxy* Proxy = InFrameAllocator.Create<FHoldoutCompositePassMaterialProxy>(Input0);
 		Proxy->Material = PostProcessMaterial;
