HoldoutComposite: Minor addition of multiply & divide merge ops.

diff --git a/Engine/Plugins/Experimental/Compositing/HoldoutComposite/Shaders/Private/HoldoutCompositeMerge.usf b/Engine/Plugins/Experimental/Compositing/HoldoutComposite/Shaders/Private/HoldoutCompositeMerge.usf
--- a/Engine/Plugins/Experimental/Compositing/HoldoutComposite/Shaders/Private/HoldoutCompositeMerge.usf
+++ b/Engine/Plugins/Experimental/Compositing/HoldoutComposite/Shaders/Private/HoldoutCompositeMerge.usf
@@ -15,6 +15,8 @@
 
 #define EHoldoutCompositePassMergeOp_Over 0
 #define EHoldoutCompositePassMergeOp_Add  1
+#define EHoldoutCompositePassMergeOp_Mul  2
+#define EHoldoutCompositePassMergeOp_Div  3
 
 SCREEN_PASS_TEXTURE_VIEWPORT(Input0)
 SCREEN_PASS_TEXTURE_VIEWPORT(Input1)
@@ -110,6 +112,21 @@ float GetOpacity(float InAlpha, uint bAlphaIsOpacity)
 	return select(bAlphaIsOpacity > 0, InAlpha, 1.0 - InAlpha);
 }
 
+float SafeDiv(float A, float B)
+{
+    return B == 0.0 ? 0.0 : A / B;
+}
+
+float4 SafeDiv(float4 A, float4 B)
+{
+    return float4(
+        SafeDiv(A[0], B[0]),
+        SafeDiv(A[1], B[1]),
+        SafeDiv(A[2], B[2]),
+        SafeDiv(A[3], B[3])
+    );
+}
+
 void MainPS(
     float4 SvPosition : SV_POSITION,
     out float4 OutColor : SV_Target0
@@ -148,6 +165,14 @@ void MainPS(
     {
         OutColor = Color0 + Color1;
     }
+    else if(MergeOp == EHoldoutCompositePassMergeOp_Mul)
+    {
+        OutColor = Color0 * Color1;
+    }
+    else if(MergeOp == EHoldoutCompositePassMergeOp_Div)
+    {
+        OutColor = SafeDiv(Color0, Color1);
+    }
     else
     {
         OutColor = Color0;
diff --git a/Engine/Plugins/Experimental/Compositing/HoldoutComposite/Source/HoldoutComposite/Public/Passes/HoldoutCompositePassMerge.h b/Engine/Plugins/Experimental/Compositing/HoldoutComposite/Source/HoldoutComposite/Public/Passes/HoldoutCompositePassMerge.h
--- a/Engine/Plugins/Experimental/Compositing/HoldoutComposite/Source/HoldoutComposite/Public/Passes/HoldoutCompositePassMerge.h
+++ b/Engine/Plugins/Experimental/Compositing/HoldoutComposite/Source/HoldoutComposite/Public/Passes/HoldoutCompositePassMerge.h
@@ -12,6 +12,8 @@ enum class EHoldoutCompositePassMergeOp : uint8
 {
 	Over,
 	Add,
+	Multiply,
+	Divide,
 };
 
 namespace UE
