Mark HoldoutComposite plugin as hidden, since it was deprecated & renamed to CompositeCore in the previous release cycle.

diff --git a/Engine/Plugins/Experimental/Compositing/HoldoutComposite/HoldoutComposite.uplugin b/Engine/Plugins/Experimental/Compositing/HoldoutComposite/HoldoutComposite.uplugin
--- a/Engine/Plugins/Experimental/Compositing/HoldoutComposite/HoldoutComposite.uplugin
+++ b/Engine/Plugins/Experimental/Compositing/HoldoutComposite/HoldoutComposite.uplugin
@@ -15,6 +15,7 @@
 	"IsBetaVersion": false,
 	"IsExperimentalVersion": true,
 	"Installed": false,
+	"Hidden": true,
 	"Modules": [
 	],
 	"Plugins": [
